fixed incorrect module dependencies in various target platforms

diff --git a/Engine/Source/Developer/Android/AndroidTargetPlatform/Private/AndroidTargetPlatformPrivatePCH.h b/Engine/Source/Developer/Android/AndroidTargetPlatform/Private/AndroidTargetPlatformPrivatePCH.h
--- a/Engine/Source/Developer/Android/AndroidTargetPlatform/Private/AndroidTargetPlatformPrivatePCH.h
+++ b/Engine/Source/Developer/Android/AndroidTargetPlatform/Private/AndroidTargetPlatformPrivatePCH.h
@@ -13,7 +13,6 @@
 #include "Core.h"
 #include "ModuleInterface.h"
 #include "ModuleManager.h"
-#include "Messaging.h"
 #include "Runtime/Core/Public/Android/AndroidProperties.h"
 
 #if WITH_ENGINE
diff --git a/Engine/Source/Developer/HTML5/HTML5TargetPlatform/HTML5TargetPlatform.Build.cs b/Engine/Source/Developer/HTML5/HTML5TargetPlatform/HTML5TargetPlatform.Build.cs
--- a/Engine/Source/Developer/HTML5/HTML5TargetPlatform/HTML5TargetPlatform.Build.cs
+++ b/Engine/Source/Developer/HTML5/HTML5TargetPlatform/HTML5TargetPlatform.Build.cs
@@ -17,10 +17,9 @@ public class HTML5TargetPlatform : ModuleRules
 			new string[] {
 				"Core",
 				"CoreUObject",
-				"Settings",
 				"Sockets"
 			}
-			);
+		);
 
 		if (UEBuildConfiguration.bCompileAgainstEngine)
 		{
@@ -28,9 +27,14 @@ public class HTML5TargetPlatform : ModuleRules
 //			PrivateIncludePathModuleNames.Add("TextureCompressor");
 		}
 
+		PrivateIncludePathModuleNames.AddRange(
+			new string[] {
+				"Settings",
+			}
+		);
+
 		PrivateIncludePaths.AddRange(
-			new string[]
-			{
+			new string[] {
 				"Developer/HTML5TargetPlatform/Classes"
 			}
 		);
diff --git a/Engine/Source/Developer/Linux/LinuxTargetPlatform/LinuxTargetPlatform.Build.cs b/Engine/Source/Developer/Linux/LinuxTargetPlatform/LinuxTargetPlatform.Build.cs
--- a/Engine/Source/Developer/Linux/LinuxTargetPlatform/LinuxTargetPlatform.Build.cs
+++ b/Engine/Source/Developer/Linux/LinuxTargetPlatform/LinuxTargetPlatform.Build.cs
@@ -8,40 +8,36 @@ public class LinuxTargetPlatform : ModuleRules
 	{
         BinariesSubFolder = "Linux";
 
-        PrivateDependencyModuleNames.AddRange(
-            new string[]
-            {
+		PublicIncludePaths.AddRange(
+			new string[] {
+				"Runtime/Core/Public/Linux"
+			}
+		);
+		
+		PrivateDependencyModuleNames.AddRange(
+            new string[] {
 				"Core",
 				"CoreUObject",
-				"Settings",
 				"TargetPlatform"
 			}
         );
 
-        PublicIncludePaths.AddRange(
-            new string[]
-			{
-				"Runtime/Core/Public/Linux"
+		PrivateIncludePathModuleNames.AddRange(
+			new string[] {
+				"Settings",
 			}
-        );
-
-        if (UEBuildConfiguration.bCompileAgainstEngine)
-        {
-            PrivateDependencyModuleNames.AddRange(
-                new string[] 
-                {
-				    "Engine"
-				}
-            );
-
-            PrivateIncludePathModuleNames.Add("TextureCompressor");
-        }
+		);
 
         PrivateIncludePaths.AddRange(
-            new string[]
-			{
+            new string[] {
 				"Developer/LinuxTargetPlatform/Classes"
 			}
         );
+
+		if (UEBuildConfiguration.bCompileAgainstEngine)
+		{
+			PrivateDependencyModuleNames.Add("Engine");
+			PrivateIncludePathModuleNames.Add("TextureCompressor");
+		}
 	}
 }
diff --git a/Engine/Source/Developer/Mac/MacTargetPlatform/MacTargetPlatform.Build.cs b/Engine/Source/Developer/Mac/MacTargetPlatform/MacTargetPlatform.Build.cs
--- a/Engine/Source/Developer/Mac/MacTargetPlatform/MacTargetPlatform.Build.cs
+++ b/Engine/Source/Developer/Mac/MacTargetPlatform/MacTargetPlatform.Build.cs
@@ -7,26 +7,29 @@ public class MacTargetPlatform : ModuleRules
 	public MacTargetPlatform(TargetInfo Target)
 	{
 		PrivateDependencyModuleNames.AddRange(
-			new string[]
-			{
+			new string[] {
 				"Core",
 				"CoreUObject",
-				"Settings",
 				"TargetPlatform"
 			}
 		);
 
-		if (UEBuildConfiguration.bCompileAgainstEngine)
-		{
-			PrivateDependencyModuleNames.Add("Engine");
-			PrivateIncludePathModuleNames.Add("TextureCompressor");
-		}
+		PrivateIncludePathModuleNames.AddRange(
+			new string[] {
+				"Settings",
+			}
+		);
 
 		PrivateIncludePaths.AddRange(
-			new string[]
-			{
+			new string[] {
 				"Developer/MacTargetPlatform/Classes"
 			}
 		);
+
+		if (UEBuildConfiguration.bCompileAgainstEngine)
+		{
+			PrivateDependencyModuleNames.Add("Engine");
+			PrivateIncludePathModuleNames.Add("TextureCompressor");
+		}
 	}
 }
diff --git a/Engine/Source/Developer/Windows/WindowsTargetPlatform/WindowsTargetPlatform.Build.cs b/Engine/Source/Developer/Windows/WindowsTargetPlatform/WindowsTargetPlatform.Build.cs
--- a/Engine/Source/Developer/Windows/WindowsTargetPlatform/WindowsTargetPlatform.Build.cs
+++ b/Engine/Source/Developer/Windows/WindowsTargetPlatform/WindowsTargetPlatform.Build.cs
@@ -7,32 +7,30 @@ public class WindowsTargetPlatform : ModuleRules
 	public WindowsTargetPlatform(TargetInfo Target)
 	{
 		PrivateDependencyModuleNames.AddRange(
-			new string[]
-			{
+			new string[] {
 				"Core",
 				"CoreUObject",
-				"Settings",
 				"TargetPlatform"
 			}
 		);
 
-		if (UEBuildConfiguration.bCompileAgainstEngine)
-		{
-			PrivateDependencyModuleNames.AddRange(
-				new string[]
-				{
-					"Engine"
-				}
-			);
-
-			PrivateIncludePathModuleNames.Add("TextureCompressor");
-		}
+		PrivateIncludePathModuleNames.AddRange(
+			new string[] {
+				"Settings",
+			}
+		);
 
 		PrivateIncludePaths.AddRange(
-			new string[]
-			{
+			new string[] {
 				"Developer/WindowsTargetPlatform/Classes"
 			}
 		);
+
+		// compile with Engine
+		if (UEBuildConfiguration.bCompileAgainstEngine)
+		{
+			PrivateDependencyModuleNames.Add("Engine");
+			PrivateIncludePathModuleNames.Add("TextureCompressor");
+		}
 	}
 }
diff --git a/Engine/Source/Developer/iOS/IOSPlatformEditor/IOSPlatformEditor.Build.cs b/Engine/Source/Developer/iOS/IOSPlatformEditor/IOSPlatformEditor.Build.cs
--- a/Engine/Source/Developer/iOS/IOSPlatformEditor/IOSPlatformEditor.Build.cs
+++ b/Engine/Source/Developer/iOS/IOSPlatformEditor/IOSPlatformEditor.Build.cs
@@ -9,8 +9,7 @@ public class IOSPlatformEditor : ModuleRules
 		BinariesSubFolder = "IOS";
 
 		PrivateDependencyModuleNames.AddRange(
-			new string[]
-			{
+			new string[] {
 				"Core",
 				"CoreUObject",
 				"InputCore",
@@ -19,12 +18,17 @@ public class IOSPlatformEditor : ModuleRules
 				"EditorStyle",
 				"PropertyEditor",
 				"SharedSettingsWidgets",
-				"Settings",
 				"SourceControl",
 				"IOSRuntimeSettings",
 			}
 		);
 
+		PrivateIncludePathModuleNames.AddRange(
+			new string[] {
+				"Settings",
+			}
+		);
+
         // this is listed above, so it isn't really dynamically loaded, this just marks it as being platform specific.
 		PlatformSpecificDynamicallyLoadedModuleNames.Add("IOSRuntimeSettings");
 	}
diff --git a/Engine/Source/Developer/iOS/IOSTargetPlatform/IOSTargetPlatform.Build.cs b/Engine/Source/Developer/iOS/IOSTargetPlatform/IOSTargetPlatform.Build.cs
--- a/Engine/Source/Developer/iOS/IOSTargetPlatform/IOSTargetPlatform.Build.cs
+++ b/Engine/Source/Developer/iOS/IOSTargetPlatform/IOSTargetPlatform.Build.cs
@@ -14,12 +14,18 @@ public class IOSTargetPlatform : ModuleRules
 				"Core",
 				"CoreUObject",
 				"TargetPlatform",
+				"LaunchDaemonMessages",
+			}
+		);
+
+		PrivateIncludePathModuleNames.AddRange(
+			new string[] {
 				"Messaging",
 				"TargetDeviceServices",
-				"LaunchDaemonMessages",
 			}
 		);
-        PlatformSpecificDynamicallyLoadedModuleNames.Add("LaunchDaemonMessages");
+
+		PlatformSpecificDynamicallyLoadedModuleNames.Add("LaunchDaemonMessages");
 
 		//This is somehow necessary for getting iOS to build on, at least, windows. It seems like the target platform is included for cooking, and thus it requirtes a bunch of other info.
 		PublicIncludePaths.AddRange(
diff --git a/Engine/Source/Programs/IOS/UnrealLaunchDaemon/UnrealLaunchDaemon.Build.cs b/Engine/Source/Programs/IOS/UnrealLaunchDaemon/UnrealLaunchDaemon.Build.cs
--- a/Engine/Source/Programs/IOS/UnrealLaunchDaemon/UnrealLaunchDaemon.Build.cs
+++ b/Engine/Source/Programs/IOS/UnrealLaunchDaemon/UnrealLaunchDaemon.Build.cs
@@ -12,13 +12,9 @@ public class UnrealLaunchDaemon : ModuleRules
 			}
 		);
 
-		PrivateIncludePaths.Add("Runtime/Launch/Private");
-
 		PrivateDependencyModuleNames.AddRange(
 			new string[] {
 				"Core",
-				"Messaging",
-				"Networking",
 				"NetworkFile",
 				"Projects",
 				"StreamingFile",
@@ -30,5 +26,13 @@ public class UnrealLaunchDaemon : ModuleRules
 				"UdpMessaging"
 			}
 		);
+
+		PrivateIncludePathModuleNames.AddRange(
+			new string[] {
+				"Messaging",
+			}
+		);
+
+		PrivateIncludePaths.Add("Runtime/Launch/Private");
 	}
 }
\ No newline at end of file
