Better handling of Linux devices in TargetManager.

- Do not add a local device unless actually running on Linux (TTP #331642)
- Use identifier as a name instead of local machine name (TTP #331801)
- Misc spelling.

diff --git a/Engine/Source/Developer/Linux/LinuxTargetPlatform/Private/LinuxTargetDevice.h b/Engine/Source/Developer/Linux/LinuxTargetPlatform/Private/LinuxTargetDevice.h
--- a/Engine/Source/Developer/Linux/LinuxTargetPlatform/Private/LinuxTargetDevice.h
+++ b/Engine/Source/Developer/Linux/LinuxTargetPlatform/Private/LinuxTargetDevice.h
@@ -31,8 +31,9 @@ public:
 	 *
 	 * @param InTargetPlatform - The target platform.
 	 */
-	FLinuxTargetDevice( const ITargetPlatform& InTargetPlatform, const FTargetDeviceId & InDeviceId )
+	FLinuxTargetDevice( const ITargetPlatform& InTargetPlatform, const FTargetDeviceId & InDeviceId, const FString & InDeviceName )
 		: TargetPlatform(InTargetPlatform)
+		, DeviceName(InDeviceName)
 		, TargetDeviceId(InDeviceId)
 	{ }
 
@@ -61,12 +62,12 @@ public:
 
 	virtual FString GetName( ) const OVERRIDE
 	{
-		return FPlatformProcess::ComputerName();
+		return DeviceName;
 	}
 
 	virtual FString GetOperatingSystemName( ) OVERRIDE
 	{
-		return TEXT("Linux-GNU");
+		return TEXT("GNU/Linux");
 	}
 
 	virtual int32 GetProcessSnapshot( TArray<FTargetDeviceProcessInfo>& OutProcessInfos ) OVERRIDE;
@@ -118,6 +119,9 @@ private:
 	// Holds a reference to the device's target platform.
 	const ITargetPlatform& TargetPlatform;
 
+	/** Device display name */
+	FString DeviceName;
+
 	/** Device id */
 	FTargetDeviceId TargetDeviceId;
 
diff --git a/Engine/Source/Developer/Linux/LinuxTargetPlatform/Private/LinuxTargetPlatform.h b/Engine/Source/Developer/Linux/LinuxTargetPlatform/Private/LinuxTargetPlatform.h
--- a/Engine/Source/Developer/Linux/LinuxTargetPlatform/Private/LinuxTargetPlatform.h
+++ b/Engine/Source/Developer/Linux/LinuxTargetPlatform/Private/LinuxTargetPlatform.h
@@ -22,8 +22,11 @@ public:
 	 * Default constructor.
 	 */
 	TLinuxTargetPlatform( )
-	{
+	{		
+#if PLATFORM_LINUX
+		// only add local device if actually running on Linux
 		LocalDevice = MakeShareable(new FLinuxTargetDevice(*this, FTargetDeviceId(PlatformName(), FPlatformProcess::ComputerName())));
+#endif
 	
 		#if WITH_ENGINE
 			FConfigCacheIni::LoadLocalIniFile(EngineSettings, TEXT("Engine"), true, *PlatformName());
@@ -43,7 +46,10 @@ public:
 	{
 		// TODO: ping all the machines in a local segment and/or try to connect to port 22 of those that respond
 		OutDevices.Reset();
-		OutDevices.Add(LocalDevice);
+		if (LocalDevice.IsValid())
+		{
+			OutDevices.Add(LocalDevice);
+		}
 	}
 
 	virtual ECompressionFlags GetBaseCompressionMethod( ) const OVERRIDE
@@ -58,7 +64,12 @@ public:
 
 	virtual ITargetDevicePtr GetDefaultDevice( ) const OVERRIDE
 	{
-		return LocalDevice;
+		if (LocalDevice.IsValid())
+		{
+			return LocalDevice;
+		}
+
+		return nullptr;
 	}
 
 	virtual ITargetDevicePtr GetDevice( const FTargetDeviceId& DeviceId ) OVERRIDE
@@ -69,7 +80,7 @@ public:
 		}
 
 		FTargetDeviceId UATFriendlyId(TEXT("linux"), DeviceId.GetDeviceName());
-		return MakeShareable(new FLinuxTargetDevice(*this, UATFriendlyId));
+		return MakeShareable(new FLinuxTargetDevice(*this, UATFriendlyId, DeviceId.GetDeviceName()));
 	}
 
 	virtual FString GetIconPath( ETargetPlatformIcons::IconType IconType ) const OVERRIDE
