Fix Linux CIS

diff --git a/Engine/Source/Developer/Linux/LinuxTargetPlatform/Private/LinuxTargetPlatform.h b/Engine/Source/Developer/Linux/LinuxTargetPlatform/Private/LinuxTargetPlatform.h
--- a/Engine/Source/Developer/Linux/LinuxTargetPlatform/Private/LinuxTargetPlatform.h
+++ b/Engine/Source/Developer/Linux/LinuxTargetPlatform/Private/LinuxTargetPlatform.h
@@ -6,6 +6,9 @@
 
 #pragma once
 
+#if WITH_ENGINE
+#include "StaticMeshResources.h"
+#endif // WITH_ENGINE
 
 /**
  * Template for Linux target platforms
