Minor fix to a stale code.

diff --git a/Engine/Source/Developer/Linux/LinuxTargetPlatform/Private/LinuxTargetPlatform.h b/Engine/Source/Developer/Linux/LinuxTargetPlatform/Private/LinuxTargetPlatform.h
--- a/Engine/Source/Developer/Linux/LinuxTargetPlatform/Private/LinuxTargetPlatform.h
+++ b/Engine/Source/Developer/Linux/LinuxTargetPlatform/Private/LinuxTargetPlatform.h
@@ -19,7 +19,7 @@ class TLinuxTargetPlatform
 {
 public:
 
-	typedef public TTargetPlatformBase<FLinuxPlatformProperties<HAS_EDITOR_DATA, IS_DEDICATED_SERVER, IS_CLIENT_ONLY> > TSuper;
+	typedef TTargetPlatformBase<FLinuxPlatformProperties<HAS_EDITOR_DATA, IS_DEDICATED_SERVER, IS_CLIENT_ONLY> > TSuper;
 
 	/**
 	 * Default constructor.
@@ -28,7 +28,8 @@ public:
 	{		
 #if PLATFORM_LINUX
 		// only add local device if actually running on Linux
-		LocalDevice = MakeShareable(new FLinuxTargetDevice(*this, FTargetDeviceId(PlatformName(), FPlatformProcess::ComputerName())));
+		FTargetDeviceId UATFriendlyId(FTargetDeviceId(TSuper::PlatformName(), FPlatformProcess::ComputerName()));
+		LocalDevice = MakeShareable(new FLinuxTargetDevice(*this, UATFriendlyId, FPlatformProcess::ComputerName()));
 #endif
 	
 		#if WITH_ENGINE
