Support multiple RHIs on Linux.

- Part of PR #306 by amigo and #ue4linux folks.

diff --git a/Engine/Config/BaseEngine.ini b/Engine/Config/BaseEngine.ini
--- a/Engine/Config/BaseEngine.ini
+++ b/Engine/Config/BaseEngine.ini
@@ -2205,6 +2205,10 @@ TilesStreamingTimeThreshold=1.0
 +TargetedRHIs=PCD3D_SM5
 +TargetedRHIs=PCD3D_SM4
 
+[/Script/LinuxTargetPlatform.LinuxTargetSettings]
++TargetedRHIs=GLSL_430
++TargetedRHIs=GLSL_150
+
 [/Script/Engine.AISystemBase]
 AISystemModuleName=AIModule
 AISystemClassName=/Script/AIModule.AISystem
diff --git a/Engine/Source/Developer/Linux/LinuxTargetPlatform/Classes/LinuxTargetSettings.h b/Engine/Source/Developer/Linux/LinuxTargetPlatform/Classes/LinuxTargetSettings.h
--- a/Engine/Source/Developer/Linux/LinuxTargetPlatform/Classes/LinuxTargetSettings.h
+++ b/Engine/Source/Developer/Linux/LinuxTargetPlatform/Classes/LinuxTargetSettings.h
@@ -12,7 +12,7 @@
 /**
  * Implements the settings for the Android target platform.
  */
-UCLASS()
+UCLASS(config=Engine, defaultconfig)
 class LINUXTARGETPLATFORM_API ULinuxTargetSettings
 	: public UObject
 {
@@ -20,4 +20,10 @@ public:
 
 	GENERATED_UCLASS_BODY()
 
+	/**
+	 * The collection of RHI's we want to support on this platform.
+	 * This is not always the full list of RHI we can support.
+	 */
+	UPROPERTY(EditAnywhere, config, Category=Rendering)
+	TArray<FString> TargetedRHIs;
 };
diff --git a/Engine/Source/Developer/Linux/LinuxTargetPlatform/LinuxTargetPlatform.Build.cs b/Engine/Source/Developer/Linux/LinuxTargetPlatform/LinuxTargetPlatform.Build.cs
--- a/Engine/Source/Developer/Linux/LinuxTargetPlatform/LinuxTargetPlatform.Build.cs
+++ b/Engine/Source/Developer/Linux/LinuxTargetPlatform/LinuxTargetPlatform.Build.cs
@@ -8,12 +8,6 @@ public class LinuxTargetPlatform : ModuleRules
 	{
         BinariesSubFolder = "Linux";
 
-		PublicIncludePaths.AddRange(
-			new string[] {
-				"Runtime/Core/Public/Linux"
-			}
-		);
-		
 		PrivateDependencyModuleNames.AddRange(
             new string[] {
 				"Core",
diff --git a/Engine/Source/Developer/Linux/LinuxTargetPlatform/Private/LinuxTargetPlatform.h b/Engine/Source/Developer/Linux/LinuxTargetPlatform/Private/LinuxTargetPlatform.h
--- a/Engine/Source/Developer/Linux/LinuxTargetPlatform/Private/LinuxTargetPlatform.h
+++ b/Engine/Source/Developer/Linux/LinuxTargetPlatform/Private/LinuxTargetPlatform.h
@@ -32,11 +32,27 @@ public:
 		LocalDevice = MakeShareable(new FLinuxTargetDevice(*this, UATFriendlyId, FPlatformProcess::ComputerName()));
 #endif
 	
-		#if WITH_ENGINE
-			FConfigCacheIni::LoadLocalIniFile(EngineSettings, TEXT("Engine"), true, *PlatformName());
-			TextureLODSettings.Initialize(EngineSettings, TEXT("SystemSettings"));
-			StaticMeshLODSettings.Initialize(EngineSettings);
-		#endif
+#if WITH_ENGINE
+		FConfigCacheIni::LoadLocalIniFile(EngineSettings, TEXT("Engine"), true, *this->PlatformName());
+		TextureLODSettings.Initialize(EngineSettings, TEXT("SystemSettings"));
+		StaticMeshLODSettings.Initialize(EngineSettings);
+
+		// Get the Target RHIs for this platform, we do not always want all those that are supported.
+		GConfig->GetArray(TEXT("/Script/LinuxTargetPlatform.LinuxTargetSettings"), TEXT("TargetedRHIs"), TargetedShaderFormats, GEngineIni);
+
+		// Gather the list of Target RHIs and filter out any that may be invalid.
+		TArray<FName> PossibleShaderFormats;
+		GetAllPossibleShaderFormats(PossibleShaderFormats);
+
+		for(int32 ShaderFormatIdx = TargetedShaderFormats.Num()-1; ShaderFormatIdx >= 0; ShaderFormatIdx--)
+		{
+			FString ShaderFormat = TargetedShaderFormats[ShaderFormatIdx];
+			if(PossibleShaderFormats.Contains(FName(*ShaderFormat)) == false)
+			{
+				TargetedShaderFormats.Remove(ShaderFormat);
+			}
+		}
+#endif // WITH_ENGINE
 	}
 
 
@@ -83,16 +99,17 @@ public:
 			return LocalDevice;
 		}
 
-		FTargetDeviceId UATFriendlyId(TEXT("linux"), DeviceId.GetDeviceName());
+		FTargetDeviceId UATFriendlyId(TEXT("Linux"), DeviceId.GetDeviceName());
 		return MakeShareable(new FLinuxTargetDevice(*this, UATFriendlyId, DeviceId.GetDeviceName()));
 	}
 
 	virtual bool IsRunningPlatform( ) const override
 	{
-		return false;
+		// Must be Linux platform as editor for this to be considered a running platform
+		return PLATFORM_LINUX && !UE_SERVER && !UE_GAME && WITH_EDITOR && HAS_EDITOR_DATA;
 	}
 
-	bool SupportsFeature(ETargetPlatformFeatures::Type Feature) const override
+	virtual bool SupportsFeature(ETargetPlatformFeatures::Type Feature) const override
 	{
 		if (Feature == ETargetPlatformFeatures::UserCredentials || Feature == ETargetPlatformFeatures::Packaging)
 		{
@@ -118,7 +135,10 @@ public:
 
 	virtual void GetAllTargetedShaderFormats( TArray<FName>& OutFormats ) const override
 	{
-		GetAllPossibleShaderFormats( OutFormats );
+		for(const FString& ShaderFormat : TargetedShaderFormats)
+		{
+			OutFormats.AddUnique(FName(*ShaderFormat));
+		}
 	}
 
 	virtual const class FStaticMeshLODSettings& GetStaticMeshLODSettings( ) const override
@@ -131,7 +151,7 @@ public:
 		if (!IS_DEDICATED_SERVER)
 		{
 			// just use the standard texture format name for this texture
-			OutFormats.Add(GetDefaultTextureFormatName(InTexture, EngineSettings));
+			OutFormats.Add(this->GetDefaultTextureFormatName(InTexture, EngineSettings));
 		}
 	}
 
@@ -143,6 +163,12 @@ public:
 	virtual FName GetWaveFormat( class USoundWave* Wave ) const override
 	{
 		static FName NAME_OGG(TEXT("OGG"));
+		static FName NAME_OPUS(TEXT("OPUS"));
+
+		if (Wave->IsStreaming())
+		{
+			return NAME_OPUS;
+		}
 
 		return NAME_OGG;
 	}
@@ -176,6 +202,9 @@ private:
 
 	// Holds static mesh LOD settings.
 	FStaticMeshLODSettings StaticMeshLODSettings;
+
+	// List of shader formats specified as targets
+	TArray<FString> TargetedShaderFormats;
 #endif // WITH_ENGINE
 
 private:
diff --git a/Engine/Source/Developer/Linux/LinuxTargetPlatform/Private/LinuxTargetPlatformModule.cpp b/Engine/Source/Developer/Linux/LinuxTargetPlatform/Private/LinuxTargetPlatformModule.cpp
--- a/Engine/Source/Developer/Linux/LinuxTargetPlatform/Private/LinuxTargetPlatformModule.cpp
+++ b/Engine/Source/Developer/Linux/LinuxTargetPlatform/Private/LinuxTargetPlatformModule.cpp
@@ -57,6 +57,9 @@ public:
 	virtual void StartupModule() override
 	{
 		TargetSettings = ConstructObject<ULinuxTargetSettings>(ULinuxTargetSettings::StaticClass(), GetTransientPackage(), "LinuxTargetSettings", RF_Standalone);
+
+		// We need to manually load the config properties here, as this module is loaded before the UObject system is setup to do this
+		GConfig->GetArray(TEXT("/Script/LinuxTargetPlatform.LinuxTargetSettings"), TEXT("TargetedRHIs"), TargetSettings->TargetedRHIs, GEngineIni);
 		TargetSettings->AddToRoot();
 
 		ISettingsModule* SettingsModule = ISettingsModule::Get();
@@ -65,7 +68,7 @@ public:
 		{
 			SettingsModule->RegisterSettings("Project", "Platforms", "Linux",
 				LOCTEXT("TargetSettingsName", "Linux"),
-				LOCTEXT("TargetSettingsDescription", "Linux platform settings description text here"),
+				LOCTEXT("TargetSettingsDescription", "Settings for Linux target platform"),
 				TargetSettings
 			);
 		}
@@ -73,9 +76,23 @@ public:
 
 	virtual void ShutdownModule() override
 	{
-		TargetSettings->RemoveFromRoot();
-	}
+		ISettingsModule* SettingsModule = ISettingsModule::Get();
+
+		if (SettingsModule != nullptr)
+		{
+			SettingsModule->UnregisterSettings("Project", "Platforms", "Linux");
+		}
 
+		if (!GExitPurge)
+		{
+			// If we're in exit purge, this object has already been destroyed
+			TargetSettings->RemoveFromRoot();
+		}
+		else
+		{
+			TargetSettings = NULL;
+		}
+	}
 	// End IModuleInterface interface
 
 
