Support local deployment on Linux, stub for other platforms.

- Part of PR #306 (with changes).

diff --git a/Engine/Source/Developer/Linux/LinuxTargetPlatform/Private/LinuxTargetDevice.cpp b/Engine/Source/Developer/Linux/LinuxTargetPlatform/Private/LinuxTargetDevice.cpp
--- a/Engine/Source/Developer/Linux/LinuxTargetPlatform/Private/LinuxTargetDevice.cpp
+++ b/Engine/Source/Developer/Linux/LinuxTargetPlatform/Private/LinuxTargetDevice.cpp
@@ -8,54 +8,148 @@
 
 #include "LinuxTargetDevice.h"
 
-
+#if PLATFORM_LINUX
+	#include <signal.h> // for process termination
+	#include <pwd.h> // for getting uid/gid
+#endif // PLATFORM_LINUX
 
 /* ITargetDevice interface
- *****************************************************************************/
+*****************************************************************************/
 
 bool FLinuxTargetDevice::Deploy( const FString& SourceFolder, FString& OutAppId )
 {
+#if PLATFORM_LINUX	// if running natively, support simplified, local deployment	
+	OutAppId = TEXT("");
+
+	FString PlatformName = TEXT("Linux");
+	FString DeploymentDir = FPaths::EngineIntermediateDir() / TEXT("Devices") / PlatformName;
+
+	// delete previous build
+	IFileManager::Get().DeleteDirectory(*DeploymentDir, false, true);
+
+	// copy files into device directory
+	TArray<FString> FileNames;
+
+	IFileManager::Get().FindFilesRecursive(FileNames, *SourceFolder, TEXT("*.*"), true, false);
+
+	for (int32 FileIndex = 0; FileIndex < FileNames.Num(); ++FileIndex)
+	{
+		const FString& SourceFilePath = FileNames[FileIndex];
+		FString DestFilePath = DeploymentDir + SourceFilePath.RightChop(SourceFolder.Len());
+
+		IFileManager::Get().Copy(*DestFilePath, *SourceFilePath);
+	}
+
+	return true;
+#else
+	// @todo: support deployment to a remote machine
+	STUBBED("FLinuxTargetDevice::Deploy");
 	return false;
+#endif // PLATFORM_LINUX
 }
 
 int32 FLinuxTargetDevice::GetProcessSnapshot( TArray<FTargetDeviceProcessInfo>& OutProcessInfos )
 {
+	STUBBED("FLinuxTargetDevice::GetProcessSnapshot");
 	return 0;
 }
 
 bool FLinuxTargetDevice::Launch( const FString& AppId, EBuildConfigurations::Type BuildConfiguration, EBuildTargets::Type BuildTarget, const FString& Params, uint32* OutProcessId )
 {
+#if PLATFORM_LINUX	// if running natively, support launching in place
+	// build executable path
+	FString PlatformName = TEXT("Linux");
+	FString ExecutablePath = FPaths::EngineIntermediateDir() / TEXT("Devices") / PlatformName / TEXT("Engine") / TEXT("Binaries") / PlatformName;
+
+	if (BuildTarget == EBuildTargets::Game)
+	{
+		ExecutablePath /= TEXT("UE4Game");
+	}
+	else if (BuildTarget == EBuildTargets::Server)
+	{
+		ExecutablePath /= TEXT("UE4Server");
+	}
+	else if (BuildTarget == EBuildTargets::Editor)
+	{
+		ExecutablePath /= TEXT("UE4Editor");
+	}
+
+	if (BuildConfiguration != EBuildConfigurations::Development)
+	{
+		ExecutablePath += FString::Printf(TEXT("-%s-%s"), *PlatformName, EBuildConfigurations::ToString(BuildConfiguration));
+	}
+
+	// launch the game
+	FProcHandle ProcessHandle = FPlatformProcess::CreateProc(*ExecutablePath, *Params, true, false, false, OutProcessId, 0, NULL, NULL);
+	return ProcessHandle.Close();
+#else
+	// @todo: support launching on a remote machine
+	STUBBED("FLinuxTargetDevice::Launch");
 	return false;
+#endif // PLATFORM_LINUX
 }
 
 bool FLinuxTargetDevice::Reboot( bool bReconnect)
 {
+	STUBBED("FLinuxTargetDevice::Reboot");
 	return false;
 }
 
 bool FLinuxTargetDevice::Run( const FString& ExecutablePath, const FString& Params, uint32* OutProcessId )
 {
+#if PLATFORM_LINUX	// if running natively, support simplified, local deployment	
+	FProcHandle ProcessHandle = FPlatformProcess::CreateProc(*ExecutablePath, *Params, true, false, false, OutProcessId, 0, NULL, NULL);
+	return ProcessHandle.Close();
+#else
+	// @todo: support remote run
+	STUBBED("FLinuxTargetDevice::Run");
 	return false;
+#endif // PLATFORM_LINUX
 }
 
 bool FLinuxTargetDevice::SupportsFeature( ETargetDeviceFeatures::Type Feature ) const
 {
+	switch (Feature)
+	{
+		case ETargetDeviceFeatures::MultiLaunch:
+			return true;
+
+		// @todo to be implemented
+		case ETargetDeviceFeatures::PowerOff:
+			return false;
+
+		// @todo to be implemented turning on remote PCs (wake on LAN)
+		case ETargetDeviceFeatures::PowerOn:
+			return false;
+
+		// @todo to be implemented
+		case ETargetDeviceFeatures::ProcessSnapshot:
+			return false;
+
+		// @todo to be implemented
+		case ETargetDeviceFeatures::Reboot:
+			return false;
+	}
+
 	return false;
 }
 
 bool FLinuxTargetDevice::SupportsSdkVersion( const FString& VersionString ) const
 {
+	STUBBED("FLinuxTargetDevice::SupportsSdkVersion");
 	return true;
 }
 
-void FLinuxTargetDevice::SetUserCredentials( const FString & InUserName, const FString & InUserPassword )
+void FLinuxTargetDevice::SetUserCredentials( const FString& InUserName, const FString& InUserPassword )
 {
+	STUBBED("FLinuxTargetDevice::SetUserCredentials");
 	UserName = InUserName;
 	UserPassword = InUserPassword;
 }
 
-bool FLinuxTargetDevice::GetUserCredentials( FString & OutUserName, FString & OutUserPassword )
+bool FLinuxTargetDevice::GetUserCredentials( FString& OutUserName, FString& OutUserPassword )
 {
+	STUBBED("FLinuxTargetDevice::GetUserCredentials");
 	OutUserName = UserName;
 	OutUserPassword = UserPassword;
 	return true;
@@ -63,6 +157,33 @@ bool FLinuxTargetDevice::GetUserCredentials( FString & OutUserName, FString & Ou
 
 bool FLinuxTargetDevice::TerminateProcess( const int32 ProcessId )
 {
+#if PLATFORM_LINUX // if running natively, just terminate the local process
+	// get process path from the ProcessId
+	const int32 ReadLinkSize = 1024;
+	char ReadLinkCmd[ReadLinkSize] = {0};
+	FCStringAnsi::Sprintf(ReadLinkCmd, "/proc/%d/exe", ProcessId);
+	char ProcessPath[ MAX_PATH + 1 ] = {0};
+	int32 Ret = readlink(ReadLinkCmd, ProcessPath, ARRAY_COUNT(ProcessPath) - 1);
+	if (Ret != -1)
+	{
+		struct stat st;
+		uid_t euid;
+		stat(ProcessPath, &st);
+		euid = geteuid(); // get effective uid of current application, as this user is asking to kill a process
+
+		// now check if we own the process
+		if (st.st_uid == euid)
+		{
+			// terminate it (will this terminate children as well because we set their pgid?)
+			kill(ProcessId, SIGTERM);
+			sleep(2); // sleep in case process still remains then send a more strict signal
+			kill(ProcessId, SIGKILL);
+			return true;
+		}
+	}
+#else
+	// @todo: support remote termination
+	STUBBED("FLinuxTargetDevice::TerminateProcess");
+#endif // PLATFORM_LINUX
 	return false;
 }
-
