Fixed compilation issue with LinuxTargetPlatform accessing protected Broadcast function.

diff --git a/Engine/Source/Developer/Linux/LinuxTargetPlatform/Private/LinuxTargetPlatform.h b/Engine/Source/Developer/Linux/LinuxTargetPlatform/Private/LinuxTargetPlatform.h
--- a/Engine/Source/Developer/Linux/LinuxTargetPlatform/Private/LinuxTargetPlatform.h
+++ b/Engine/Source/Developer/Linux/LinuxTargetPlatform/Private/LinuxTargetPlatform.h
@@ -19,8 +19,10 @@ template<bool HAS_EDITOR_DATA, bool IS_DEDICATED_SERVER, bool IS_CLIENT_ONLY>
 class TLinuxTargetPlatform
 	: public TTargetPlatformBase<FLinuxPlatformProperties<HAS_EDITOR_DATA, IS_DEDICATED_SERVER, IS_CLIENT_ONLY> >
 {
+private:
+	typedef TLinuxTargetPlatform<HAS_EDITOR_DATA, IS_DEDICATED_SERVER, IS_CLIENT_ONLY> TThisType;
 public:
-
+	
 	typedef FLinuxPlatformProperties<HAS_EDITOR_DATA, IS_DEDICATED_SERVER, IS_CLIENT_ONLY> TProperties;
 	typedef TTargetPlatformBase<TProperties> TSuper;
 
@@ -221,13 +223,13 @@ public:
 		return TProperties::GetVariantPriority();
 	}
 
-	DECLARE_DERIVED_EVENT(FAndroidTargetPlatform, ITargetPlatform::FOnTargetDeviceDiscovered, FOnTargetDeviceDiscovered);
+	DECLARE_DERIVED_EVENT(TThisType, ITargetPlatform::FOnTargetDeviceDiscovered, FOnTargetDeviceDiscovered);
 	virtual FOnTargetDeviceDiscovered& OnDeviceDiscovered( ) override
 	{
 		return DeviceDiscoveredEvent;
 	}
-
-	DECLARE_DERIVED_EVENT(FAndroidTargetPlatform, ITargetPlatform::FOnTargetDeviceLost, FOnTargetDeviceLost);
+	
+	DECLARE_DERIVED_EVENT(TThisType, ITargetPlatform::FOnTargetDeviceLost, FOnTargetDeviceLost);
 	virtual FOnTargetDeviceLost& OnDeviceLost( ) override
 	{
 		return DeviceLostEvent;
