Fix UFE compilation for Linux.

diff --git a/Engine/Source/Developer/Linux/LinuxTargetPlatform/Private/LinuxTargetPlatformPrivatePCH.h b/Engine/Source/Developer/Linux/LinuxTargetPlatform/Private/LinuxTargetPlatformPrivatePCH.h
--- a/Engine/Source/Developer/Linux/LinuxTargetPlatform/Private/LinuxTargetPlatformPrivatePCH.h
+++ b/Engine/Source/Developer/Linux/LinuxTargetPlatform/Private/LinuxTargetPlatformPrivatePCH.h
@@ -7,7 +7,8 @@
  *****************************************************************************/
 
 #include "Core.h"
-#include "Runtime/Core/Public/Linux/LinuxPlatformProperties.h"
+#include "CoreUObject.h"
+#include "LinuxPlatformProperties.h"
 
 #if WITH_ENGINE
 	#include "Engine.h"
