Fix build (incorrect #include).

diff --git a/Engine/Source/Developer/Linux/LinuxTargetPlatform/Private/LinuxTargetPlatformPrivatePCH.h b/Engine/Source/Developer/Linux/LinuxTargetPlatform/Private/LinuxTargetPlatformPrivatePCH.h
--- a/Engine/Source/Developer/Linux/LinuxTargetPlatform/Private/LinuxTargetPlatformPrivatePCH.h
+++ b/Engine/Source/Developer/Linux/LinuxTargetPlatform/Private/LinuxTargetPlatformPrivatePCH.h
@@ -8,7 +8,7 @@
 
 #include "Core.h"
 #include "CoreUObject.h"
-#include "LinuxPlatformProperties.h"
+#include "Runtime/Core/Public/Linux/LinuxPlatformProperties.h"
 
 #if WITH_ENGINE
 	#include "Engine.h"
