Fix usage of a deprecated ConstructObject().

diff --git a/Engine/Source/Developer/Linux/LinuxTargetPlatform/Private/LinuxTargetPlatformModule.cpp b/Engine/Source/Developer/Linux/LinuxTargetPlatform/Private/LinuxTargetPlatformModule.cpp
--- a/Engine/Source/Developer/Linux/LinuxTargetPlatform/Private/LinuxTargetPlatformModule.cpp
+++ b/Engine/Source/Developer/Linux/LinuxTargetPlatform/Private/LinuxTargetPlatformModule.cpp
@@ -49,7 +49,7 @@ public:
 
 	virtual void StartupModule() override
 	{
-		TargetSettings = ConstructObject<ULinuxTargetSettings>(ULinuxTargetSettings::StaticClass(), GetTransientPackage(), "LinuxTargetSettings", RF_Standalone);
+		TargetSettings = NewObject<ULinuxTargetSettings>(GetTransientPackage(), "LinuxTargetSettings", RF_Standalone);
 
 		// We need to manually load the config properties here, as this module is loaded before the UObject system is setup to do this
 		GConfig->GetArray(TEXT("/Script/LinuxTargetPlatform.LinuxTargetSettings"), TEXT("TargetedRHIs"), TargetSettings->TargetedRHIs, GEngineIni);
