Linux: UFE: store and manage multiple devices.

- Also fixes UFE crash on adding the device with the same properties (UE-14245).

diff --git a/Engine/Source/Developer/Linux/LinuxTargetPlatform/Private/LinuxTargetPlatform.h b/Engine/Source/Developer/Linux/LinuxTargetPlatform/Private/LinuxTargetPlatform.h
--- a/Engine/Source/Developer/Linux/LinuxTargetPlatform/Private/LinuxTargetPlatform.h
+++ b/Engine/Source/Developer/Linux/LinuxTargetPlatform/Private/LinuxTargetPlatform.h
@@ -67,12 +67,17 @@ public:
 
 	virtual bool AddDevice(const FString& DeviceName, bool bDefault) override
 	{
-		FTargetDeviceId UATFriendlyId(TEXT("Linux"), DeviceName);
-		FLinuxTargetDevicePtr Device = MakeShareable(new FLinuxTargetDevice(*this, UATFriendlyId, DeviceName));
+		FLinuxTargetDevicePtr& Device = Devices.FindOrAdd(DeviceName);
+
 		if (Device.IsValid())
 		{
-			DeviceDiscoveredEvent.Broadcast(Device.ToSharedRef());
+			// do not allow duplicates
+			return false;
 		}
+
+		FTargetDeviceId UATFriendlyId(TEXT("Linux"), DeviceName);
+		Device = MakeShareable(new FLinuxTargetDevice(*this, UATFriendlyId, DeviceName));
+		DeviceDiscoveredEvent.Broadcast(Device.ToSharedRef());
 		return true;
 	}
 
@@ -84,6 +89,11 @@ public:
 		{
 			OutDevices.Add(LocalDevice);
 		}
+
+		for (const auto & DeviceIter : Devices)
+		{
+			OutDevices.Add(DeviceIter.Value);
+		}
 	}
 
 	virtual ECompressionFlags GetBaseCompressionMethod( ) const override
@@ -112,6 +122,15 @@ public:
 		{
 			return LocalDevice;
 		}
+
+		for (const auto & DeviceIter : Devices)
+		{
+			if (DeviceId == DeviceIter.Value->GetId())
+			{
+				return DeviceIter.Value;
+			}
+		}
+
 		return nullptr;
 	}
 
@@ -246,6 +265,9 @@ private:
 
 	// Holds the local device.
 	FLinuxTargetDevicePtr LocalDevice;
+	// Holds a map of valid devices.
+	TMap<FString, FLinuxTargetDevicePtr> Devices;
+
 
 #if WITH_ENGINE
 	// Holds the Engine INI settings for quick use.
