Linux: fix default Targeted RHIs settings (UE-59372).

- Removed GL3 (bit rotten).
- GL4 no longer "experimental".

diff --git a/Engine/Source/Developer/Linux/LinuxTargetPlatform/Private/LinuxTargetPlatform.h b/Engine/Source/Developer/Linux/LinuxTargetPlatform/Private/LinuxTargetPlatform.h
--- a/Engine/Source/Developer/Linux/LinuxTargetPlatform/Private/LinuxTargetPlatform.h
+++ b/Engine/Source/Developer/Linux/LinuxTargetPlatform/Private/LinuxTargetPlatform.h
@@ -247,12 +247,10 @@ public:
 		{
 			static FName NAME_VULKAN_SM5(TEXT("SF_VULKAN_SM5"));
 			static FName NAME_GLSL_430(TEXT("GLSL_430"));
-			static FName NAME_GLSL_150(TEXT("GLSL_150"));
 			static FName NAME_VULKAN_ES31(TEXT("SF_VULKAN_ES31"));
 
 			OutFormats.AddUnique(NAME_VULKAN_SM5);
 			OutFormats.AddUnique(NAME_GLSL_430);
-			OutFormats.AddUnique(NAME_GLSL_150);
 			OutFormats.AddUnique(NAME_VULKAN_ES31);
 		}
 	}
diff --git a/Engine/Source/Editor/DetailCustomizations/Private/LinuxTargetSettingsDetails.cpp b/Engine/Source/Editor/DetailCustomizations/Private/LinuxTargetSettingsDetails.cpp
--- a/Engine/Source/Editor/DetailCustomizations/Private/LinuxTargetSettingsDetails.cpp
+++ b/Engine/Source/Editor/DetailCustomizations/Private/LinuxTargetSettingsDetails.cpp
@@ -63,7 +63,7 @@ static FText GetFriendlyNameFromRHIName(const FString& InRHIName)
 	}
 	else if (InRHIName == TEXT("GLSL_430"))
 	{
-		FriendlyRHIName = LOCTEXT("OpenGL4", "OpenGL 4 (SM5, Experimental)");
+		FriendlyRHIName = LOCTEXT("OpenGL4", "OpenGL 4 (SM5)");
 	}
 	else if (InRHIName == TEXT("SF_VULKAN_ES31_ANDROID") || InRHIName == TEXT("SF_VULKAN_ES31"))
 	{
