Linux: enable audio streaming, don't use ADPCM.

diff --git a/Engine/Source/Developer/Linux/LinuxTargetPlatform/Private/LinuxTargetPlatform.h b/Engine/Source/Developer/Linux/LinuxTargetPlatform/Private/LinuxTargetPlatform.h
--- a/Engine/Source/Developer/Linux/LinuxTargetPlatform/Private/LinuxTargetPlatform.h
+++ b/Engine/Source/Developer/Linux/LinuxTargetPlatform/Private/LinuxTargetPlatform.h
@@ -318,13 +318,22 @@ public:
 	virtual FName GetWaveFormat( const class USoundWave* Wave ) const override
 	{
 		static FName NAME_OGG(TEXT("OGG"));
+		static FName NAME_OPUS(TEXT("OPUS"));
+
+		if (Wave->IsStreaming())
+		{
+			return NAME_OPUS;
+		}
+
 		return NAME_OGG;
 	}
 
 	virtual void GetAllWaveFormats(TArray<FName>& OutFormats) const override
 	{
 		static FName NAME_OGG(TEXT("OGG"));
+		static FName NAME_OPUS(TEXT("OPUS"));
 		OutFormats.Add(NAME_OGG);
+		OutFormats.Add(NAME_OPUS);
 	}
 
 	virtual FPlatformAudioCookOverrides* GetAudioCompressionSettings() const override
diff --git a/Engine/Source/Runtime/Core/Public/Linux/LinuxPlatformProperties.h b/Engine/Source/Runtime/Core/Public/Linux/LinuxPlatformProperties.h
--- a/Engine/Source/Runtime/Core/Public/Linux/LinuxPlatformProperties.h
+++ b/Engine/Source/Runtime/Core/Public/Linux/LinuxPlatformProperties.h
@@ -139,4 +139,9 @@ struct FLinuxPlatformProperties
 	{
 		return true;
 	}
+
+	static FORCEINLINE bool SupportsAudioStreaming()
+	{
+		return !IsServerOnly();
+	}
 };
diff --git a/Engine/Source/Runtime/Linux/AudioMixerSDL/Private/AudioMixerPlatformSDL.cpp b/Engine/Source/Runtime/Linux/AudioMixerSDL/Private/AudioMixerPlatformSDL.cpp
--- a/Engine/Source/Runtime/Linux/AudioMixerSDL/Private/AudioMixerPlatformSDL.cpp
+++ b/Engine/Source/Runtime/Linux/AudioMixerSDL/Private/AudioMixerPlatformSDL.cpp
@@ -328,13 +328,7 @@ namespace Audio
 		}
 
 		static FName NAME_OGG(TEXT("OGG"));
-		if (InSoundWave->HasCompressedData(NAME_OGG))
-		{
-			return NAME_OGG;
-		}
-
-		static FName NAME_ADPCM(TEXT("ADPCM"));
-		return NAME_ADPCM;
+		return NAME_OGG;
 	}
 
 	bool FMixerPlatformSDL::HasCompressedAudioInfoClass(USoundWave* InSoundWave)
