Missing file

#rnx

diff --git a/Engine/Source/Developer/Linux/LinuxTargetPlatform/Private/LinuxTargetDevice.h b/Engine/Source/Developer/Linux/LinuxTargetPlatform/Private/LinuxTargetDevice.h
--- a/Engine/Source/Developer/Linux/LinuxTargetPlatform/Private/LinuxTargetDevice.h
+++ b/Engine/Source/Developer/Linux/LinuxTargetPlatform/Private/LinuxTargetDevice.h
@@ -264,7 +264,7 @@ public:
 		const int32 ReadLinkSize = 1024;
 		char ReadLinkCmd[ReadLinkSize] = { 0 };
 		FCStringAnsi::Sprintf(ReadLinkCmd, "/proc/%lld/exe", ProcessId);
-		char ProcessPath[MAX_PATH + 1] = { 0 };
+		char ProcessPath[UNIX_MAX_PATH + 1] = { 0 };
 		int32 Ret = readlink(ReadLinkCmd, ProcessPath, ARRAY_COUNT(ProcessPath) - 1);
 		if (Ret != -1)
 		{
