More CIS fixes.

#rnx

diff --git a/Engine/Source/Developer/Linux/LinuxTargetPlatform/Private/LinuxTargetDevice.h b/Engine/Source/Developer/Linux/LinuxTargetPlatform/Private/LinuxTargetDevice.h
--- a/Engine/Source/Developer/Linux/LinuxTargetPlatform/Private/LinuxTargetDevice.h
+++ b/Engine/Source/Developer/Linux/LinuxTargetPlatform/Private/LinuxTargetDevice.h
@@ -150,15 +150,15 @@ public:
 		FString PlatformName = TEXT("Linux");
 		FString ExecutablePath = FPaths::EngineIntermediateDir() / TEXT("Devices") / PlatformName / TEXT("Engine") / TEXT("Binaries") / PlatformName;
 
-		if (BuildTarget == EBuildTargetType::Game)
+		if (TargetType == EBuildTargetType::Game)
 		{
 			ExecutablePath /= TEXT("UE4Game");
 		}
-		else if (BuildTarget == EBuildTargetType::Server)
+		else if (TargetType == EBuildTargetType::Server)
 		{
 			ExecutablePath /= TEXT("UE4Server");
 		}
-		else if (BuildTarget == EBuildTargetType::Editor)
+		else if (TargetType == EBuildTargetType::Editor)
 		{
 			ExecutablePath /= TEXT("UE4Editor");
 		}
