Fix a couple of CIS errors.

#rnx

diff --git a/Engine/Source/Developer/Linux/LinuxTargetPlatform/Private/LinuxTargetPlatform.h b/Engine/Source/Developer/Linux/LinuxTargetPlatform/Private/LinuxTargetPlatform.h
--- a/Engine/Source/Developer/Linux/LinuxTargetPlatform/Private/LinuxTargetPlatform.h
+++ b/Engine/Source/Developer/Linux/LinuxTargetPlatform/Private/LinuxTargetPlatform.h
@@ -215,7 +215,7 @@ public:
 			}
 
 			FText Reason;
-			if (RequiresTempTarget(bProjectHasCode, Configuration, bRequiresAssetNativization, Reason))
+			if (this->RequiresTempTarget(bProjectHasCode, Configuration, bRequiresAssetNativization, Reason))
 			{
 				ReadyToBuild |= ETargetPlatformReadyStatus::PluginsUnsupported;
 			}
