Don't add local launch device for Linux AArch64 platform

Local device right now will never be an Arm64 device.

diff --git a/Engine/Source/Developer/Linux/LinuxTargetPlatform/Private/LinuxTargetPlatform.h b/Engine/Source/Developer/Linux/LinuxTargetPlatform/Private/LinuxTargetPlatform.h
--- a/Engine/Source/Developer/Linux/LinuxTargetPlatform/Private/LinuxTargetPlatform.h
+++ b/Engine/Source/Developer/Linux/LinuxTargetPlatform/Private/LinuxTargetPlatform.h
@@ -48,8 +48,11 @@ public:
 #endif // WITH_ENGINE
 	{
 #if PLATFORM_LINUX
-		// only add local device if actually running on Linux
-		LocalDevice = MakeShareable(new FLinuxTargetDevice(*this, FPlatformProcess::ComputerName(), nullptr));
+		if (!TProperties::IsAArch64())
+		{
+			// only add local device if actually running on Linux
+			LocalDevice = MakeShareable(new FLinuxTargetDevice(*this, FPlatformProcess::ComputerName(), nullptr));
+		}
 #endif
 
 #if WITH_ENGINE
