Fix AArch64 -> Arm64

diff --git a/Engine/Source/Developer/Linux/LinuxArm64TargetPlatform/LinuxArm64TargetPlatform.Build.cs b/Engine/Source/Developer/Linux/LinuxArm64TargetPlatform/LinuxArm64TargetPlatform.Build.cs
--- a/Engine/Source/Developer/Linux/LinuxArm64TargetPlatform/LinuxArm64TargetPlatform.Build.cs
+++ b/Engine/Source/Developer/Linux/LinuxArm64TargetPlatform/LinuxArm64TargetPlatform.Build.cs
@@ -2,11 +2,11 @@
 
 using UnrealBuildTool;
 
-public class LinuxAArch64TargetPlatform : ModuleRules
+public class LinuxArm64TargetPlatform : ModuleRules
 {
-    public LinuxAArch64TargetPlatform(ReadOnlyTargetRules Target) : base(Target)
+    public LinuxArm64TargetPlatform(ReadOnlyTargetRules Target) : base(Target)
     {
-        BinariesSubFolder = "LinuxAArch64";
+        BinariesSubFolder = "LinuxArm64";
 
         PrivateDependencyModuleNames.AddRange(
             new string[] {
diff --git a/Engine/Source/Developer/Linux/LinuxArm64TargetPlatform/Private/LinuxArm64TargetPlatformModule.cpp b/Engine/Source/Developer/Linux/LinuxArm64TargetPlatform/Private/LinuxArm64TargetPlatformModule.cpp
--- a/Engine/Source/Developer/Linux/LinuxArm64TargetPlatform/Private/LinuxArm64TargetPlatformModule.cpp
+++ b/Engine/Source/Developer/Linux/LinuxArm64TargetPlatform/Private/LinuxArm64TargetPlatformModule.cpp
@@ -1,7 +1,7 @@
 // Copyright Epic Games, Inc. All Rights Reserved.
 
 /*=============================================================================
-	LinuxAArch64NoEditorTargetPlatformModule.cpp: Implements the FLinuxAArch64NoEditorTargetPlatformModule class.
+	LinuxArm64NoEditorTargetPlatformModule.cpp: Implements the FLinuxArm64NoEditorTargetPlatformModule class.
 =============================================================================*/
 
 #include "CoreMinimal.h"
@@ -16,7 +16,7 @@
 /**
  * Module for the Linux target platforms
  */
-class FLinuxAArch64TargetPlatformModule
+class FLinuxArm64TargetPlatformModule
 	: public ITargetPlatformModule
 {
 public:
@@ -33,4 +33,4 @@ public:
 };
 
 
-IMPLEMENT_MODULE(FLinuxAArch64TargetPlatformModule, LinuxAArch64TargetPlatform);
+IMPLEMENT_MODULE(FLinuxArm64TargetPlatformModule, LinuxArm64TargetPlatform);
