Don't restrict the possible shader formats for server only (Fix for Linux compile error)

diff --git a/Engine/Source/Developer/Linux/LinuxTargetPlatform/Private/LinuxTargetPlatform.h b/Engine/Source/Developer/Linux/LinuxTargetPlatform/Private/LinuxTargetPlatform.h
--- a/Engine/Source/Developer/Linux/LinuxTargetPlatform/Private/LinuxTargetPlatform.h
+++ b/Engine/Source/Developer/Linux/LinuxTargetPlatform/Private/LinuxTargetPlatform.h
@@ -235,15 +235,11 @@ public:
 
 	virtual void GetAllPossibleShaderFormats( TArray<FName>& OutFormats ) const override
 	{
-		// no shaders needed for dedicated server target
-		if (this->AllowAudioVisualData())
-		{
-			static FName NAME_VULKAN_SM5(TEXT("SF_VULKAN_SM5"));
-			static FName NAME_VULKAN_ES31(TEXT("SF_VULKAN_ES31"));
+		static FName NAME_VULKAN_SM5(TEXT("SF_VULKAN_SM5"));
+		static FName NAME_VULKAN_ES31(TEXT("SF_VULKAN_ES31"));
 
-			OutFormats.AddUnique(NAME_VULKAN_SM5);
-			OutFormats.AddUnique(NAME_VULKAN_ES31);
-		}
+		OutFormats.AddUnique(NAME_VULKAN_SM5);
+		OutFormats.AddUnique(NAME_VULKAN_ES31);
 	}
 
 	virtual void GetAllTargetedShaderFormats( TArray<FName>& OutFormats ) const override
