Make Linux Target Platform logic for shader formats match Windows by not targeting any platforms when server only. Fixes an issue where shaders could get cooked for targeted RHIs on Linux Server.

diff --git a/Engine/Source/Developer/Linux/LinuxTargetPlatform/Private/LinuxTargetPlatform.h b/Engine/Source/Developer/Linux/LinuxTargetPlatform/Private/LinuxTargetPlatform.h
--- a/Engine/Source/Developer/Linux/LinuxTargetPlatform/Private/LinuxTargetPlatform.h
+++ b/Engine/Source/Developer/Linux/LinuxTargetPlatform/Private/LinuxTargetPlatform.h
@@ -258,11 +258,14 @@
 
 	virtual void GetAllPossibleShaderFormats( TArray<FName>& OutFormats ) const override
 	{
-		static FName NAME_VULKAN_SM5(TEXT("SF_VULKAN_SM5"));
-		static FName NAME_VULKAN_ES31(TEXT("SF_VULKAN_ES31"));
+		if (!TProperties::IsServerOnly())
+		{
+			static FName NAME_VULKAN_SM5(TEXT("SF_VULKAN_SM5"));
+			static FName NAME_VULKAN_ES31(TEXT("SF_VULKAN_ES31"));
 
-		OutFormats.AddUnique(NAME_VULKAN_SM5);
-		OutFormats.AddUnique(NAME_VULKAN_ES31);
+			OutFormats.AddUnique(NAME_VULKAN_SM5);
+			OutFormats.AddUnique(NAME_VULKAN_ES31);
+		}
 	}
 
 	virtual void GetAllTargetedShaderFormats( TArray<FName>& OutFormats ) const override
