Moving some LinuxTargetPlatform and MacTargetPlatform private files to Public since they're used by other modules.

diff --git a/Engine/Source/Developer/CookedEditor/CookedEditor.Build.cs b/Engine/Source/Developer/CookedEditor/CookedEditor.Build.cs
--- a/Engine/Source/Developer/CookedEditor/CookedEditor.Build.cs
+++ b/Engine/Source/Developer/CookedEditor/CookedEditor.Build.cs
@@ -1,6 +1,5 @@
 // Copyright Epic Games, Inc. All Rights Reserved.
 
-using System.IO;
 using UnrealBuildTool;
 
 public class CookedEditor : ModuleRules
@@ -30,15 +29,24 @@ public class CookedEditor : ModuleRules
 
 		PublicIncludePathModuleNames.Add("WindowsTargetPlatform");
 
-		// we currently need to pull in desktop templates that are private
-		// @todo move these out to Public so they can be extended like this
-		string EnginePath = Path.GetFullPath(Target.RelativeEnginePath);
-		PublicIncludePaths.AddRange(
-			new string[] {
-				Path.Combine(EnginePath, "Source/Developer/Linux/LinuxTargetPlatform/Private"),
-				Path.Combine(EnginePath, "Source/Developer/Mac/MacTargetPlatform/Private"),
-			}
-		);
+		if (IsPlatformAvailable(UnrealTargetPlatform.Linux))
+		{
+			PublicDefinitions.Add("COOKEDEDITOR_WITH_LINUXTARGETPLATFORM=1");
+			PublicIncludePathModuleNames.Add("LinuxTargetPlatform");
+		}
+		else
+		{
+			PublicDefinitions.Add("COOKEDEDITOR_WITH_LINUXTARGETPLATFORM=0");
+		}
 
+		if (IsPlatformAvailable(UnrealTargetPlatform.Mac))
+		{
+			PublicDefinitions.Add("COOKEDEDITOR_WITH_MACTARGETPLATFORM=1");
+			PublicIncludePathModuleNames.Add("MacTargetPlatform");
+		}
+		else
+		{
+			PublicDefinitions.Add("COOKEDEDITOR_WITH_MACTARGETPLATFORM=0");
+		}
 	}
 }
diff --git a/Engine/Source/Developer/CookedEditor/Public/CookedEditorPackageManager.h b/Engine/Source/Developer/CookedEditor/Public/CookedEditorPackageManager.h
--- a/Engine/Source/Developer/CookedEditor/Public/CookedEditorPackageManager.h
+++ b/Engine/Source/Developer/CookedEditor/Public/CookedEditorPackageManager.h
@@ -8,17 +8,24 @@
 
 #include "GenericWindowsTargetPlatform.h"
 #include "Windows/WindowsPlatformProperties.h"
-#include "LinuxTargetPlatform.h"
-#include "GenericMacTargetPlatform.h"
+
 typedef TGenericWindowsTargetPlatform<FWindowsPlatformProperties<false, false, false>> FWindowsEditorTargetPlatformParent;
+
+#if COOKEDEDITOR_WITH_LINUXTARGETPLATFORM
+#include "LinuxTargetPlatform.h"
 typedef TLinuxTargetPlatform<FLinuxPlatformProperties<false, false, false, false>> FLinuxEditorTargetPlatformParent;
+#endif
+
+#if COOKEDEDITOR_WITH_MACTARGETPLATFORM
+#include "GenericMacTargetPlatform.h"
+typedef TGenericMacTargetPlatform<false, false, false> FMacEditorTargetPlatformParent;
+#endif
 
 #if PLATFORM_WINDOWS
 typedef FWindowsEditorTargetPlatformParent FHostPlatformEditorTargetPlatformParent;
 #elif PLATFORM_LINUX
 typedef FLinuxEditorTargetPlatformParent FHostPlatformEditorTargetPlatformParent;
 #elif PLATFORM_MAC
-typedef TGenericMacTargetPlatform<false, false, false> FMacEditorTargetPlatformParent;
 typedef FMacEditorTargetPlatformParent FHostPlatformEditorTargetPlatformParent;
 #endif
 
diff --git a/Engine/Source/Developer/Linux/LinuxArm64TargetPlatform/LinuxArm64TargetPlatform.Build.cs b/Engine/Source/Developer/Linux/LinuxArm64TargetPlatform/LinuxArm64TargetPlatform.Build.cs
--- a/Engine/Source/Developer/Linux/LinuxArm64TargetPlatform/LinuxArm64TargetPlatform.Build.cs
+++ b/Engine/Source/Developer/Linux/LinuxArm64TargetPlatform/LinuxArm64TargetPlatform.Build.cs
@@ -27,10 +27,6 @@ public class LinuxArm64TargetPlatform : ModuleRules
             PrivateIncludePathModuleNames.Add("TextureCompressor");
         }
 
-        PrivateIncludePaths.AddRange(
-            new string[] {
-				"Developer/Linux/LinuxTargetPlatform/Private"
-			}
-        );
+		PrivateIncludePathModuleNames.Add("LinuxTargetPlatform");
     }
 }
diff --git a/Engine/Source/Developer/Linux/LinuxTargetPlatform/LinuxTargetPlatform.Build.cs b/Engine/Source/Developer/Linux/LinuxTargetPlatform/LinuxTargetPlatform.Build.cs
--- a/Engine/Source/Developer/Linux/LinuxTargetPlatform/LinuxTargetPlatform.Build.cs
+++ b/Engine/Source/Developer/Linux/LinuxTargetPlatform/LinuxTargetPlatform.Build.cs
@@ -20,7 +20,7 @@ public class LinuxTargetPlatform : ModuleRules
 
 		PrivateIncludePathModuleNames.AddRange(
 			new string[] {
-				"Settings",
+				"Settings"
 			}
 		);
 
diff --git a/Engine/Source/Developer/Linux/LinuxTargetPlatform/Private/LinuxTargetDevice.h b/Engine/Source/Developer/Linux/LinuxTargetPlatform/Private/LinuxTargetDevice.h
deleted file mode 100644
diff --git a/Engine/Source/Developer/Linux/LinuxTargetPlatform/Private/LinuxTargetPlatform.h b/Engine/Source/Developer/Linux/LinuxTargetPlatform/Private/LinuxTargetPlatform.h
deleted file mode 100644
diff --git a/Engine/Source/Developer/Linux/LinuxTargetPlatform/Public/LinuxTargetDevice.h b/Engine/Source/Developer/Linux/LinuxTargetPlatform/Public/LinuxTargetDevice.h
new file mode 100644
--- /dev/null
+++ b/Engine/Source/Developer/Linux/LinuxTargetPlatform/Public/LinuxTargetDevice.h
@@ -0,0 +1,219 @@
+// Copyright Epic Games, Inc. All Rights Reserved.
+
+#pragma once
+
+#include "CoreMinimal.h"
+#include "Misc/Paths.h"
+#include "HAL/FileManager.h"
+#include "HAL/PlatformProcess.h"
+#include "Interfaces/ITargetPlatform.h"
+#include "Interfaces/ITargetDevice.h"
+#include "Interfaces/TargetDeviceId.h"
+
+#if PLATFORM_LINUX
+	#include <signal.h>
+	#include <pwd.h>
+#endif // PLATFORM_LINUX
+
+class IFileManager;
+struct FProcHandle;
+
+
+/**
+ * Type definition for shared pointers to instances of FLinuxTargetDevice.
+ */
+typedef TSharedPtr<class FLinuxTargetDevice, ESPMode::ThreadSafe> FLinuxTargetDevicePtr;
+
+/**
+ * Type definition for shared references to instances of FLinuxTargetDevice.
+ */
+typedef TSharedRef<class FLinuxTargetDevice, ESPMode::ThreadSafe> FLinuxTargetDeviceRef;
+
+
+/**
+ * Implements a Linux target device.
+ */
+class FLinuxTargetDevice
+	: public ITargetDevice
+{
+public:
+
+	/**
+	 * Creates and initializes a new device for the specified target platform.
+	 *
+	 * @param InTargetPlatform - The target platform.
+	 */
+	FLinuxTargetDevice( const ITargetPlatform& InTargetPlatform, const FString& InDeviceName, TFunction<void()> InSavePlatformDevices)
+		: TargetPlatform(InTargetPlatform)
+		, DeviceName(InDeviceName)
+		, SavePlatformDevices(InSavePlatformDevices)
+	{ }
+
+	/**
+	 * Minimal constructor for use with SteamDeck devices */
+	FLinuxTargetDevice(const ITargetPlatform& InTargetPlatform)
+		: FLinuxTargetDevice(InTargetPlatform, TEXT("UnknownName"), nullptr) 
+	{ }
+
+public:
+
+	virtual bool Connect( ) override
+	{
+		return true;
+	}
+
+	virtual void Disconnect( ) override
+	{ }
+
+	virtual ETargetDeviceTypes GetDeviceType( ) const override
+	{
+		return ETargetDeviceTypes::Desktop;
+	}
+
+	virtual FTargetDeviceId GetId() const override
+	{
+		return FTargetDeviceId(TargetPlatform.PlatformName(), GetName());
+	}
+
+	virtual FString GetName( ) const override
+	{
+		return DeviceName;
+	}
+
+	virtual FString GetOperatingSystemName( ) override
+	{
+		return TEXT("GNU/Linux");
+	}
+
+	virtual int32 GetProcessSnapshot( TArray<FTargetDeviceProcessInfo>& OutProcessInfos ) override
+	{
+		STUBBED("FLinuxTargetDevice::GetProcessSnapshot");
+		return 0;
+	}
+
+	virtual const class ITargetPlatform& GetTargetPlatform( ) const override
+	{
+		return TargetPlatform;
+	}
+
+	virtual bool IsConnected( ) override
+	{
+		return true;
+	}
+
+	virtual bool IsDefault( ) const override
+	{
+		return true;
+	}
+
+	virtual bool PowerOff( bool Force ) override
+	{
+		return false;
+	}
+
+	virtual bool PowerOn( ) override
+	{
+		return false;
+	}
+
+	virtual bool Reboot( bool bReconnect = false ) override
+	{
+		STUBBED("FLinuxTargetDevice::Reboot");
+		return false;
+	}
+
+	virtual bool SupportsFeature( ETargetDeviceFeatures Feature ) const override
+	{
+		switch (Feature)
+		{
+		case ETargetDeviceFeatures::MultiLaunch:
+			return true;
+
+			// @todo to be implemented
+		case ETargetDeviceFeatures::PowerOff:
+			return false;
+
+			// @todo to be implemented turning on remote PCs (wake on LAN)
+		case ETargetDeviceFeatures::PowerOn:
+			return false;
+
+			// @todo to be implemented
+		case ETargetDeviceFeatures::ProcessSnapshot:
+			return false;
+
+			// @todo to be implemented
+		case ETargetDeviceFeatures::Reboot:
+			return false;
+		}
+
+		return false;
+	}
+
+	virtual void SetUserCredentials( const FString& InUserName, const FString& InUserPassword ) override
+	{
+		UserName = InUserName;
+		UserPassword = InUserPassword;
+
+		if (SavePlatformDevices)
+		{
+			SavePlatformDevices();
+		}
+	}
+
+	virtual bool GetUserCredentials( FString& OutUserName, FString& OutUserPassword ) override
+	{
+		OutUserName = UserName;
+		OutUserPassword = UserPassword;
+		return true;
+	}
+
+	virtual bool TerminateProcess( const int64 ProcessId ) override
+	{
+#if PLATFORM_LINUX // if running natively, just terminate the local process
+		// get process path from the ProcessId
+		const int32 ReadLinkSize = 1024;
+		char ReadLinkCmd[ReadLinkSize] = { 0 };
+		FCStringAnsi::Sprintf(ReadLinkCmd, "/proc/%lld/exe", ProcessId);
+		char ProcessPath[UNIX_MAX_PATH + 1] = { 0 };
+		int32 Ret = readlink(ReadLinkCmd, ProcessPath, UE_ARRAY_COUNT(ProcessPath) - 1);
+		if (Ret != -1)
+		{
+			struct stat st;
+			uid_t euid;
+			stat(ProcessPath, &st);
+			euid = geteuid(); // get effective uid of current application, as this user is asking to kill a process
+
+							  // now check if we own the process
+			if (st.st_uid == euid)
+			{
+				// terminate it (will this terminate children as well because we set their pgid?)
+				kill(ProcessId, SIGTERM);
+				sleep(2); // sleep in case process still remains then send a more strict signal
+				kill(ProcessId, SIGKILL);
+				return true;
+			}
+		}
+#else
+		// @todo: support remote termination
+		STUBBED("FLinuxTargetDevice::TerminateProcess");
+#endif // PLATFORM_LINUX
+		return false;
+	}
+
+protected:
+
+	// Holds a reference to the device's target platform.
+	const ITargetPlatform& TargetPlatform;
+
+	/** Device display name */
+	FString DeviceName;
+
+	/** User name on the remote machine */
+	FString UserName;
+
+	/** User password on the remote machine */
+	FString UserPassword;
+
+	/** Target platform function to save device state */
+	TFunction<void()> SavePlatformDevices;
+};
diff --git a/Engine/Source/Developer/Linux/LinuxTargetPlatform/Public/LinuxTargetPlatform.h b/Engine/Source/Developer/Linux/LinuxTargetPlatform/Public/LinuxTargetPlatform.h
new file mode 100644
--- /dev/null
+++ b/Engine/Source/Developer/Linux/LinuxTargetPlatform/Public/LinuxTargetPlatform.h
@@ -0,0 +1,483 @@
+// Copyright Epic Games, Inc. All Rights Reserved.
+
+/*=============================================================================
+	LinuxTargetPlatform.h: Declares the FLinuxTargetPlatform class.
+=============================================================================*/
+
+#pragma once
+
+#include "CoreMinimal.h"
+#include "Interfaces/TargetDeviceId.h"
+#include "Misc/Paths.h"
+#include "Misc/ConfigCacheIni.h"
+#include "Interfaces/ITargetPlatform.h"
+#include "Common/TargetPlatformBase.h"
+#include "SteamDeck/SteamDeckDevice.h"
+#include "AnalyticsEventAttribute.h"
+
+
+#if WITH_ENGINE
+#include "AudioCompressionSettings.h"
+#include "Sound/SoundWave.h"
+#include "StaticMeshResources.h"
+#include "TextureResource.h"
+#endif // WITH_ENGINE
+#include "InstalledPlatformInfo.h"
+#include "LinuxTargetDevice.h"
+#include "Linux/LinuxPlatformProperties.h"
+
+#define LOCTEXT_NAMESPACE "TLinuxTargetPlatform"
+
+class UTextureLODSettings;
+
+/**
+ * Template for Linux target platforms
+ */
+template<typename TProperties>
+class TLinuxTargetPlatform
+	: public TTargetPlatformBase<TProperties>
+{
+public:
+
+	typedef TTargetPlatformBase<TProperties> TSuper;
+
+	/**
+	 * Default constructor.
+	 */
+	TLinuxTargetPlatform( )
+#if WITH_ENGINE
+		: bChangingDeviceConfig(false)
+#endif // WITH_ENGINE
+	{
+#if PLATFORM_LINUX
+		if (!TProperties::IsArm64())
+		{
+			// only add local device if actually running on Linux
+			LocalDevice = MakeShareable(new FLinuxTargetDevice(*this, FPlatformProcess::ComputerName(), nullptr));
+		}
+#endif
+
+#if WITH_ENGINE
+		TextureLODSettings = nullptr;
+		StaticMeshLODSettings.Initialize(this);
+
+		InitDevicesFromConfig();
+
+		if (!TProperties::IsArm64())
+		{
+			SteamDevices = TSteamDeckDevice<FLinuxTargetDevice>::DiscoverDevices(*this, TEXT("Native Linux"));
+		}
+
+
+		// Get the Target RHIs for this platform, we do not always want all those that are supported.
+		TArray<FName> TargetedShaderFormats;
+		TLinuxTargetPlatform::GetAllTargetedShaderFormats(TargetedShaderFormats);
+
+		// If we are targeting ES 2.0/3.1, we also must cook encoded HDR reflection captures
+		static FName NAME_SF_VULKAN_ES31(TEXT("SF_VULKAN_ES31"));
+		static FName NAME_OPENGL_150_ES3_1(TEXT("GLSL_150_ES31"));
+		bRequiresEncodedHDRReflectionCaptures = TargetedShaderFormats.Contains(NAME_SF_VULKAN_ES31)
+			|| TargetedShaderFormats.Contains(NAME_OPENGL_150_ES3_1);
+#endif // WITH_ENGINE
+	}
+
+
+public:
+
+	//~ Begin ITargetPlatform Interface
+
+	virtual void EnableDeviceCheck(bool OnOff) override {}
+
+	virtual bool AddDevice(const FString& DeviceName, bool bDefault) override
+	{
+		return AddDevice(DeviceName, TEXT(""), TEXT(""), TEXT(""), bDefault);
+	}
+
+	virtual bool AddDevice(const FString& DeviceName, const FString& DeviceUserFriendlyName, const FString& Username, const FString& Password, bool bDefault) override
+	{
+		FLinuxTargetDevicePtr& Device = Devices.FindOrAdd(DeviceName);
+
+		if (Device.IsValid())
+		{
+			// do not allow duplicates
+			return false;
+		}
+
+		Device = MakeShareable(new FLinuxTargetDevice(*this, DeviceName,
+#if WITH_ENGINE
+			[&]() { SaveDevicesToConfig(); }));
+		SaveDevicesToConfig();	// this will do the right thing even if AddDevice() was called from InitDevicesFromConfig
+#else
+			nullptr));
+#endif // WITH_ENGINE
+
+		if (!Username.IsEmpty() || !Password.IsEmpty())
+		{
+			Device->SetUserCredentials(Username, Password);
+		}
+
+		this->OnDeviceDiscovered().Broadcast(Device.ToSharedRef());
+		return true;
+	}
+
+
+	virtual void GetAllDevices( TArray<ITargetDevicePtr>& OutDevices ) const override
+	{
+		// TODO: ping all the machines in a local segment and/or try to connect to port 22 of those that respond
+		OutDevices.Reset();
+		if (LocalDevice.IsValid())
+		{
+			OutDevices.Add(LocalDevice);
+		}
+
+		for (const auto & DeviceIter : Devices)
+		{
+			OutDevices.Add(DeviceIter.Value);
+		}
+
+		for (const ITargetDevicePtr& SteamDeck : SteamDevices)
+		{
+			if (SteamDeck.IsValid())
+			{
+				OutDevices.Add(SteamDeck);
+			}
+		}
+	}
+
+	virtual bool GenerateStreamingInstallManifest(const TMultiMap<FString, int32>& PakchunkMap, const TSet<int32>& PakchunkIndicesInUse) const override
+	{
+		return true;
+	}
+
+	virtual ITargetDevicePtr GetDefaultDevice( ) const override
+	{
+		if (LocalDevice.IsValid())
+		{
+			return LocalDevice;
+		}
+
+		return nullptr;
+	}
+
+	virtual ITargetDevicePtr GetDevice( const FTargetDeviceId& DeviceId ) override
+	{
+		if (LocalDevice.IsValid() && (DeviceId == LocalDevice->GetId()))
+		{
+			return LocalDevice;
+		}
+
+		for (const auto & DeviceIter : Devices)
+		{
+			if (DeviceId == DeviceIter.Value->GetId())
+			{
+				return DeviceIter.Value;
+			}
+		}
+
+		for (const ITargetDevicePtr& SteamDeck : SteamDevices)
+		{
+			if (SteamDeck.IsValid() && DeviceId == SteamDeck->GetId())
+			{
+				return SteamDeck;
+			}
+		}
+
+		return nullptr;
+	}
+
+	virtual bool IsRunningPlatform( ) const override
+	{
+		// Must be Linux platform as editor for this to be considered a running platform
+		return PLATFORM_LINUX && !UE_SERVER && !UE_GAME && WITH_EDITOR && TProperties::HasEditorOnlyData();
+	}
+
+	virtual bool SupportsFeature(ETargetPlatformFeatures Feature) const override
+	{
+		if (Feature == ETargetPlatformFeatures::UserCredentials || Feature == ETargetPlatformFeatures::Packaging)
+		{
+			return true;
+		}
+
+		return TSuper::SupportsFeature(Feature);
+	}
+
+	virtual bool IsSdkInstalled(bool bProjectHasCode, FString& OutDocumentationPath) const override
+	{
+		if (!PLATFORM_LINUX)
+		{
+			// check for LINUX_MULTIARCH_ROOT or for legacy LINUX_ROOT when targeting Linux from Win/Mac
+
+			// proceed with any value for MULTIARCH root, because checking exact architecture is not possible at this point
+			FString ToolchainMultiarchRoot = FPlatformMisc::GetEnvironmentVariable(TEXT("LINUX_MULTIARCH_ROOT"));
+			if (ToolchainMultiarchRoot.Len() > 0 && FPaths::DirectoryExists(ToolchainMultiarchRoot))
+			{
+				return true;
+			}
+
+			// else check for legacy LINUX_ROOT
+			FString ToolchainCompiler = FPlatformMisc::GetEnvironmentVariable(TEXT("LINUX_ROOT"));
+			if (PLATFORM_WINDOWS)
+			{
+				ToolchainCompiler += "/bin/clang++.exe";
+			}
+			else if (PLATFORM_MAC)
+			{
+				ToolchainCompiler += "/bin/clang++";
+			}
+			else
+			{
+				checkf(false, TEXT("Unable to target Linux on an unknown platform."));
+				return false;
+			}
+
+			return FPaths::FileExists(ToolchainCompiler);
+		}
+
+		return true;
+	}
+
+	virtual int32 CheckRequirements(bool bProjectHasCode, EBuildConfiguration Configuration, bool bRequiresAssetNativization, FString& OutTutorialPath, FString& OutDocumentationPath, FText& CustomizedLogMessage) const override
+	{
+		int32 ReadyToBuild = TSuper::CheckRequirements(bProjectHasCode, Configuration, bRequiresAssetNativization, OutTutorialPath, OutDocumentationPath, CustomizedLogMessage);
+
+		// do not support code/plugins in Installed builds if the required libs aren't bundled (on Windows/Mac)
+		if (!PLATFORM_LINUX && !FInstalledPlatformInfo::Get().IsValidPlatform(TSuper::GetPlatformInfo().UBTPlatformString, EProjectType::Code))
+		{
+			if (bProjectHasCode)
+			{
+				ReadyToBuild |= ETargetPlatformReadyStatus::CodeUnsupported;
+			}
+
+			FText Reason;
+			if (this->RequiresTempTarget(bProjectHasCode, Configuration, bRequiresAssetNativization, Reason))
+			{
+				ReadyToBuild |= ETargetPlatformReadyStatus::PluginsUnsupported;
+			}
+		}
+
+		return ReadyToBuild;
+	}
+
+	virtual void GetAllPossibleShaderFormats( TArray<FName>& OutFormats ) const override
+	{
+		if (!TProperties::IsServerOnly())
+		{
+			static FName NAME_VULKAN_SM5(TEXT("SF_VULKAN_SM5"));
+			static FName NAME_VULKAN_SM6(TEXT("SF_VULKAN_SM6"));
+			static FName NAME_VULKAN_ES31(TEXT("SF_VULKAN_ES31"));
+
+			OutFormats.AddUnique(NAME_VULKAN_SM5);
+			OutFormats.AddUnique(NAME_VULKAN_SM6);
+			OutFormats.AddUnique(NAME_VULKAN_ES31);
+		}
+	}
+
+	virtual void GetAllTargetedShaderFormats( TArray<FName>& OutFormats ) const override
+	{
+		// Get the Target RHIs for this platform, we do not always want all those that are supported. (reload in case user changed in the editor)
+		TArray<FString>TargetedShaderFormats;
+		GConfig->GetArray(TEXT("/Script/LinuxTargetPlatform.LinuxTargetSettings"), TEXT("TargetedRHIs"), TargetedShaderFormats, GEngineIni);
+
+		// Gather the list of Target RHIs and filter out any that may be invalid.
+		TArray<FName> PossibleShaderFormats;
+		GetAllPossibleShaderFormats(PossibleShaderFormats);
+
+		for (int32 ShaderFormatIdx = TargetedShaderFormats.Num() - 1; ShaderFormatIdx >= 0; ShaderFormatIdx--)
+		{
+			FString ShaderFormat = TargetedShaderFormats[ShaderFormatIdx];
+			if (PossibleShaderFormats.Contains(FName(*ShaderFormat)) == false)
+			{
+				TargetedShaderFormats.RemoveAt(ShaderFormatIdx);
+			}
+		}
+
+		for(const FString& ShaderFormat : TargetedShaderFormats)
+		{
+			OutFormats.AddUnique(FName(*ShaderFormat));
+		}
+	}
+
+	virtual void GetPlatformSpecificProjectAnalytics( TArray<FAnalyticsEventAttribute>& AnalyticsParamArray ) const override
+	{
+		TSuper::GetPlatformSpecificProjectAnalytics(AnalyticsParamArray);
+
+		TSuper::AppendAnalyticsEventConfigArray(AnalyticsParamArray, TEXT("/Script/LinuxTargetPlatform.LinuxTargetSettings"), TEXT("TargetedRHIs"), GEngineIni);
+	}
+
+#if WITH_ENGINE
+	virtual void GetReflectionCaptureFormats(TArray<FName>& OutFormats) const override
+	{
+		if (bRequiresEncodedHDRReflectionCaptures)
+		{
+			OutFormats.Add(FName(TEXT("EncodedHDR")));
+		}
+
+		OutFormats.Add(FName(TEXT("FullHDR")));
+	}
+
+	virtual const class FStaticMeshLODSettings& GetStaticMeshLODSettings( ) const override
+	{
+		return StaticMeshLODSettings;
+	}
+
+	virtual void GetTextureFormats( const UTexture* InTexture, TArray< TArray<FName> >& OutFormats) const override
+	{
+		if (this->AllowAudioVisualData())
+		{
+			// just use the standard texture format name for this texture
+			GetDefaultTextureFormatNamePerLayer(OutFormats.AddDefaulted_GetRef(), this, InTexture, true, 4, true);
+		}
+	}
+
+	virtual void GetAllTextureFormats(TArray<FName>& OutFormats) const override
+	{
+		if (this->AllowAudioVisualData())
+		{
+			// just use the standard texture format name for this texture
+			GetAllDefaultTextureFormats(this, OutFormats);
+		}
+	}
+
+	virtual const UTextureLODSettings& GetTextureLODSettings() const override
+	{
+		return *TextureLODSettings;
+	}
+
+	virtual void RegisterTextureLODSettings(const UTextureLODSettings* InTextureLODSettings) override
+	{
+		TextureLODSettings = InTextureLODSettings;
+	}
+
+#endif //WITH_ENGINE
+
+	virtual bool SupportsVariants() const override
+	{
+		return true;
+	}
+
+	virtual float GetVariantPriority() const override
+	{
+		return TProperties::GetVariantPriority();
+	}
+
+	//~ End ITargetPlatform Interface
+
+protected:
+
+#if WITH_ENGINE
+	/** Whether we're in process of changing device config - if yes, we will prevent recurrent calls. */
+	bool bChangingDeviceConfig;
+
+	void InitDevicesFromConfig()
+	{
+		if (bChangingDeviceConfig)
+		{
+			return;
+		}
+		bChangingDeviceConfig = true;
+
+		int NumDevices = 0;
+		for (;; ++NumDevices)
+		{
+			FString DeviceName, DeviceUser, DevicePass;
+
+			FString DeviceBaseKey(FString::Printf(TEXT("LinuxTargetPlatfrom_%s_Device_%d"), *TSuper::PlatformName(), NumDevices));
+			FString DeviceNameKey = DeviceBaseKey + TEXT("_Name");
+			if (!GConfig->GetString(TEXT("/Script/LinuxTargetPlatform.LinuxTargetSettings"), *DeviceNameKey, DeviceName, GEngineIni))
+			{
+				// no such device
+				break;
+			}
+
+			if (!AddDevice(DeviceName, false))
+			{
+				break;
+			}
+
+			// set credentials, if any
+			FString DeviceUserKey = DeviceBaseKey + TEXT("_User");
+			if (GConfig->GetString(TEXT("/Script/LinuxTargetPlatform.LinuxTargetSettings"), *DeviceUserKey, DeviceUser, GEngineIni))
+			{
+				FString DevicePassKey = DeviceBaseKey + TEXT("_Pass");
+				if (GConfig->GetString(TEXT("/Script/LinuxTargetPlatform.LinuxTargetSettings"), *DevicePassKey, DevicePass, GEngineIni))
+				{
+					for (const auto & DeviceIter : Devices)
+					{
+						ITargetDevicePtr Device = DeviceIter.Value;
+						if (Device.IsValid() && Device->GetId().GetDeviceName() == DeviceName)
+						{
+							Device->SetUserCredentials(DeviceUser, DevicePass);
+						}
+					}
+				}
+			}
+		}
+
+		bChangingDeviceConfig = false;
+	}
+
+	void SaveDevicesToConfig()
+	{
+		if (bChangingDeviceConfig)
+		{
+			return;
+		}
+		bChangingDeviceConfig = true;
+
+		int DeviceIndex = 0;
+		for (const auto & DeviceIter : Devices)
+		{
+			ITargetDevicePtr Device = DeviceIter.Value;
+
+			FString DeviceBaseKey(FString::Printf(TEXT("LinuxTargetPlatfrom_%s_Device_%d"), *TSuper::PlatformName(), DeviceIndex));
+			FString DeviceNameKey = DeviceBaseKey + TEXT("_Name");
+
+			if (Device.IsValid())
+			{
+				FString DeviceName = Device->GetId().GetDeviceName();
+				// do not save a local device on Linux or it will be duplicated
+				if (PLATFORM_LINUX && DeviceName == FPlatformProcess::ComputerName())
+				{
+					continue;
+				}
+
+				GConfig->SetString(TEXT("/Script/LinuxTargetPlatform.LinuxTargetSettings"), *DeviceNameKey, *DeviceName, GEngineIni);
+
+				FString DeviceUser, DevicePass;
+				if (Device->GetUserCredentials(DeviceUser, DevicePass))
+				{
+					FString DeviceUserKey = DeviceBaseKey + TEXT("_User");
+					FString DevicePassKey = DeviceBaseKey + TEXT("_Pass");
+
+					GConfig->SetString(TEXT("/Script/LinuxTargetPlatform.LinuxTargetSettings"), *DeviceUserKey, *DeviceUser, GEngineIni);
+					GConfig->SetString(TEXT("/Script/LinuxTargetPlatform.LinuxTargetSettings"), *DevicePassKey, *DevicePass, GEngineIni);
+				}
+
+				++DeviceIndex;	// needs to be incremented here since we cannot allow gaps
+			}
+		}
+
+		bChangingDeviceConfig = false;
+	}
+#endif // WITH_ENGINE
+
+	// Holds the local device.
+	FLinuxTargetDevicePtr LocalDevice;
+	// Holds a map of valid devices.
+	TMap<FString, FLinuxTargetDevicePtr> Devices;
+	TArray<ITargetDevicePtr> SteamDevices;
+
+
+#if WITH_ENGINE
+	// Holds the texture LOD settings.
+	const UTextureLODSettings* TextureLODSettings;
+
+	// Holds static mesh LOD settings.
+	FStaticMeshLODSettings StaticMeshLODSettings;
+
+	// True if the project requires encoded HDR reflection captures
+	bool bRequiresEncodedHDRReflectionCaptures;
+#endif // WITH_ENGINE
+};
+
+#undef LOCTEXT_NAMESPACE
diff --git a/Engine/Source/Developer/Mac/MacTargetPlatform/MacTargetPlatform.Build.cs b/Engine/Source/Developer/Mac/MacTargetPlatform/MacTargetPlatform.Build.cs
--- a/Engine/Source/Developer/Mac/MacTargetPlatform/MacTargetPlatform.Build.cs
+++ b/Engine/Source/Developer/Mac/MacTargetPlatform/MacTargetPlatform.Build.cs
@@ -17,7 +17,7 @@ public class MacTargetPlatform : ModuleRules
 
 		PrivateIncludePathModuleNames.AddRange(
 			new string[] {
-				"Settings",
+				"Settings"
 			}
 		);
 
diff --git a/Engine/Source/Developer/Mac/MacTargetPlatform/Private/GenericMacTargetPlatform.h b/Engine/Source/Developer/Mac/MacTargetPlatform/Private/GenericMacTargetPlatform.h
deleted file mode 100644
diff --git a/Engine/Source/Developer/Mac/MacTargetPlatform/Private/LocalMacTargetDevice.h b/Engine/Source/Developer/Mac/MacTargetPlatform/Private/LocalMacTargetDevice.h
deleted file mode 100644
diff --git a/Engine/Source/Developer/Mac/MacTargetPlatform/Public/GenericMacTargetPlatform.h b/Engine/Source/Developer/Mac/MacTargetPlatform/Public/GenericMacTargetPlatform.h
new file mode 100644
--- /dev/null
+++ b/Engine/Source/Developer/Mac/MacTargetPlatform/Public/GenericMacTargetPlatform.h
@@ -0,0 +1,248 @@
+// Copyright Epic Games, Inc. All Rights Reserved.
+
+/*=============================================================================
+	GenericMacTargetPlatform.h: Declares the TGenericMacTargetPlatform class template.
+=============================================================================*/
+
+#pragma once
+
+#include "CoreMinimal.h"
+#include "Common/TargetPlatformBase.h"
+#include "Mac/MacPlatformProperties.h"
+#include "Misc/ConfigCacheIni.h"
+#include "LocalMacTargetDevice.h"
+#include "AnalyticsEventAttribute.h"
+
+
+#if WITH_ENGINE
+#include "AudioCompressionSettings.h"
+#include "Sound/SoundWave.h"
+#include "TextureResource.h"
+#include "StaticMeshResources.h"
+#endif // WITH_ENGINE
+
+#define LOCTEXT_NAMESPACE "TGenericMacTargetPlatform"
+
+/**
+ * Template for Mac target platforms
+ */
+template<bool HAS_EDITOR_DATA, bool IS_DEDICATED_SERVER, bool IS_CLIENT_ONLY>
+class TGenericMacTargetPlatform
+	: public TTargetPlatformBase<FMacPlatformProperties<HAS_EDITOR_DATA, IS_DEDICATED_SERVER, IS_CLIENT_ONLY> >
+{
+public:
+
+	typedef FMacPlatformProperties<HAS_EDITOR_DATA, IS_DEDICATED_SERVER, IS_CLIENT_ONLY> TProperties;
+	typedef TTargetPlatformBase<TProperties> TSuper;
+
+	/**
+	 * Default constructor.
+	 */
+	TGenericMacTargetPlatform( )
+	{
+#if PLATFORM_MAC
+		// only add local device if actually running on Mac
+		LocalDevice = MakeShareable(new FLocalMacTargetDevice(*this));
+#endif
+
+		#if WITH_ENGINE
+			TextureLODSettings = nullptr;
+			StaticMeshLODSettings.Initialize(this);
+		#endif
+	}
+
+public:
+
+	//~ Begin ITargetPlatform Interface
+
+	virtual void EnableDeviceCheck(bool OnOff) override {}
+
+	virtual void GetAllDevices( TArray<ITargetDevicePtr>& OutDevices ) const override
+	{
+		OutDevices.Reset();
+		if (LocalDevice.IsValid())
+		{
+			OutDevices.Add(LocalDevice);
+		}
+	}
+
+	virtual bool GenerateStreamingInstallManifest(const TMultiMap<FString, int32>& PakchunkMap, const TSet<int32>& PakchunkIndicesInUse) const override
+	{
+		return true;
+	}
+
+	virtual ITargetDevicePtr GetDefaultDevice( ) const override
+	{
+		if (LocalDevice.IsValid())
+		{
+			return LocalDevice;
+		}
+
+		return nullptr;
+	}
+
+	virtual ITargetDevicePtr GetDevice( const FTargetDeviceId& DeviceId )
+	{
+		if (LocalDevice.IsValid() && (DeviceId == LocalDevice->GetId()))
+		{
+			return LocalDevice;
+		}
+
+		return NULL;
+	}
+
+	virtual bool IsRunningPlatform( ) const override
+	{
+		// Must be Mac platform as editor for this to be considered a running platform
+		return PLATFORM_MAC && !UE_SERVER && !UE_GAME && WITH_EDITOR && HAS_EDITOR_DATA;
+	}
+
+	virtual bool SupportsFeature( ETargetPlatformFeatures Feature ) const override
+	{
+		// we currently do not have a build target for MacServer
+		if (Feature == ETargetPlatformFeatures::Packaging)
+		{
+			return (HAS_EDITOR_DATA || !IS_DEDICATED_SERVER);
+		}
+
+		return TSuper::SupportsFeature(Feature);
+	}
+
+	virtual void GetAllPossibleShaderFormats( TArray<FName>& OutFormats ) const override
+	{
+		// no shaders needed for dedicated server target
+		if (!IS_DEDICATED_SERVER)
+		{
+			static FName NAME_SF_METAL_SM5(TEXT("SF_METAL_SM5"));
+			OutFormats.AddUnique(NAME_SF_METAL_SM5);
+			static FName NAME_SF_METAL_MACES3_1(TEXT("SF_METAL_MACES3_1"));
+			OutFormats.AddUnique(NAME_SF_METAL_MACES3_1);
+			static FName NAME_SF_METAL_MRT_MAC(TEXT("SF_METAL_MRT_MAC"));
+			OutFormats.AddUnique(NAME_SF_METAL_MRT_MAC);
+		}
+	}
+
+	virtual void GetAllTargetedShaderFormats(TArray<FName>& OutFormats) const override
+	{
+		// Get the Target RHIs for this platform, we do not always want all those that are supported.
+		TArray<FString>TargetedShaderFormats;
+		GConfig->GetArray(TEXT("/Script/MacTargetPlatform.MacTargetSettings"), TEXT("TargetedRHIs"), TargetedShaderFormats, GEngineIni);
+
+		// Gather the list of Target RHIs and filter out any that may be invalid.
+		TArray<FName> PossibleShaderFormats;
+		GetAllPossibleShaderFormats(PossibleShaderFormats);
+
+		for (int32 ShaderFormatIdx = TargetedShaderFormats.Num() - 1; ShaderFormatIdx >= 0; ShaderFormatIdx--)
+		{
+			FString ShaderFormat = TargetedShaderFormats[ShaderFormatIdx];
+			if (PossibleShaderFormats.Contains(FName(*ShaderFormat)) == false)
+			{
+				TargetedShaderFormats.RemoveAt(ShaderFormatIdx);
+			}
+		}
+
+		for(const FString& ShaderFormat : TargetedShaderFormats)
+		{
+			OutFormats.AddUnique(FName(*ShaderFormat));
+		}
+	}
+
+	virtual void GetPlatformSpecificProjectAnalytics( TArray<FAnalyticsEventAttribute>& AnalyticsParamArray ) const override
+	{
+		TSuper::GetPlatformSpecificProjectAnalytics(AnalyticsParamArray);
+
+		TSuper::AppendAnalyticsEventConfigArray(AnalyticsParamArray, TEXT("/Script/MacTargetPlatform.MacTargetSettings"), TEXT("TargetedRHIs"), GEngineIni);
+	}
+
+#if WITH_ENGINE
+	virtual const class FStaticMeshLODSettings& GetStaticMeshLODSettings( ) const override
+	{
+		return StaticMeshLODSettings;
+	}
+
+	virtual void GetTextureFormats( const UTexture* Texture, TArray< TArray<FName> >& OutFormats) const override
+	{
+		if (!IS_DEDICATED_SERVER)
+		{
+			// just use the standard texture format name for this texture (with DX11 support)
+			GetDefaultTextureFormatNamePerLayer(OutFormats.AddDefaulted_GetRef(), this, Texture, true, 4, true);
+		}
+	}
+
+	virtual void GetAllTextureFormats(TArray<FName>& OutFormats) const override
+	{
+		if (!IS_DEDICATED_SERVER)
+		{
+			// just use the standard texture format name for this texture (with DX11 support)
+			GetAllDefaultTextureFormats(this, OutFormats);
+		}
+	}
+
+	virtual bool SupportsLQCompressionTextureFormat() const override { return false; };
+
+	virtual const UTextureLODSettings& GetTextureLODSettings() const override
+	{
+		return *TextureLODSettings;
+	}
+
+	virtual void RegisterTextureLODSettings(const UTextureLODSettings* InTextureLODSettings) override
+	{
+		TextureLODSettings = InTextureLODSettings;
+	}
+
+	virtual bool CanSupportRemoteShaderCompile() const override
+	{
+		return true;
+	}
+	
+	virtual void GetShaderCompilerDependencies(TArray<FString>& OutDependencies) const override
+	{
+		FTargetPlatformBase::AddDependencySCArrayHelper(OutDependencies, TEXT("Binaries/ThirdParty/ShaderConductor/Mac/libdxcompiler.dylib"));
+		FTargetPlatformBase::AddDependencySCArrayHelper(OutDependencies, TEXT("Binaries/ThirdParty/ShaderConductor/Mac/libShaderConductor.dylib"));
+	}
+
+#endif //WITH_ENGINE
+
+
+	virtual bool SupportsVariants() const override
+	{
+		return true;
+	}
+
+
+	virtual float GetVariantPriority() const override
+	{
+		return TProperties::GetVariantPriority();
+	}
+
+	static FORCEINLINE bool SupportsRayTracing()
+	{
+		return true;
+	}
+
+	virtual bool UsesRayTracing() const override
+	{
+		bool bEnableRayTracing = false;
+		GConfig->GetBool(TEXT("/Script/MacTargetPlatform.MacTargetSettings"), TEXT("bEnableRayTracing"), bEnableRayTracing, GEngineIni);
+	 
+		return bEnableRayTracing;
+	}
+	//~ End ITargetPlatform Interface
+
+private:
+
+	// Holds the local device.
+	ITargetDevicePtr LocalDevice;
+
+#if WITH_ENGINE
+	// Holds the texture LOD settings.
+	const UTextureLODSettings* TextureLODSettings;
+
+	// Holds the static mesh LOD settings.
+	FStaticMeshLODSettings StaticMeshLODSettings;
+
+#endif // WITH_ENGINE
+
+};
+
+#undef LOCTEXT_NAMESPACE
diff --git a/Engine/Source/Developer/Mac/MacTargetPlatform/Public/LocalMacTargetDevice.h b/Engine/Source/Developer/Mac/MacTargetPlatform/Public/LocalMacTargetDevice.h
new file mode 100644
--- /dev/null
+++ b/Engine/Source/Developer/Mac/MacTargetPlatform/Public/LocalMacTargetDevice.h
@@ -0,0 +1,140 @@
+// Copyright Epic Games, Inc. All Rights Reserved.
+
+/*=============================================================================
+	LocalMacTargetDevice.h: Declares the TLocalMacTargetDevice class.
+=============================================================================*/
+
+#pragma once
+
+#include "CoreMinimal.h"
+#include "Interfaces/ITargetDevice.h"
+#include "HAL/FileManager.h"
+#include "HAL/PlatformProcess.h"
+
+/**
+ * Template for local Mac target devices.
+ */
+class FLocalMacTargetDevice
+	: public ITargetDevice
+{
+public:
+
+	/**
+	 * Creates and initializes a new device for the specified target platform.
+	 *
+	 * @param InTargetPlatform - The target platform.
+	 */
+	FLocalMacTargetDevice( const ITargetPlatform& InTargetPlatform )
+		: TargetPlatform(InTargetPlatform)
+	{ }
+
+
+public:
+
+	virtual bool Connect( ) override
+	{
+		return true;
+	}
+
+	virtual void Disconnect( )
+	{ }
+
+	virtual ETargetDeviceTypes GetDeviceType( ) const override
+	{
+		return ETargetDeviceTypes::Desktop;
+	}
+
+	virtual FTargetDeviceId GetId( ) const override
+	{
+		return FTargetDeviceId(TargetPlatform.PlatformName(), GetName());
+	}
+
+	virtual FString GetName( ) const override
+	{
+		return FPlatformProcess::ComputerName();
+	}
+
+	virtual FString GetOperatingSystemName( ) override
+	{
+		return TEXT("macOS");
+	}
+
+	virtual int32 GetProcessSnapshot( TArray<FTargetDeviceProcessInfo>& OutProcessInfos ) override
+	{
+		// @todo Mac: implement process snapshots
+		return 0;
+	}
+
+	virtual const class ITargetPlatform& GetTargetPlatform( ) const override
+	{
+		return TargetPlatform;
+	}
+
+	virtual bool IsConnected( )
+	{
+		return true;
+	}
+
+	virtual bool IsDefault( ) const override
+	{
+		return true;
+	}
+
+	virtual bool PowerOff( bool Force ) override
+	{
+		return false;
+	}
+
+	virtual bool PowerOn( ) override
+	{
+		return false;
+	}
+
+	virtual bool Reboot( bool bReconnect = false ) override
+	{
+#if PLATFORM_MAC
+		NSAppleScript* Script = [[NSAppleScript alloc] initWithSource:@"tell application \"System Events\" to restart"];
+		NSDictionary* ErrorDict = [NSDictionary dictionary];
+		[Script executeAndReturnError: &ErrorDict];
+#endif
+		return true;
+	}
+
+	virtual bool SupportsFeature( ETargetDeviceFeatures Feature ) const override
+	{
+		switch (Feature)
+		{
+		case ETargetDeviceFeatures::MultiLaunch:
+			return true;
+
+		// @todo Mac: implement process snapshots
+		case ETargetDeviceFeatures::ProcessSnapshot:
+			return false;
+
+		case ETargetDeviceFeatures::Reboot:
+			return true;
+		}
+
+		return false;
+	}
+
+	virtual void SetUserCredentials( const FString& UserName, const FString& UserPassword ) override
+	{
+	}
+
+	virtual bool GetUserCredentials( FString& OutUserName, FString& OutUserPassword ) override
+	{
+		return false;
+	}
+
+	virtual bool TerminateProcess( const int64 ProcessId ) override
+	{
+		return false;
+	}
+
+
+private:
+
+	// Holds a reference to the device's target platform.
+	const ITargetPlatform& TargetPlatform;
+};
