Desktop Platforms TP Refactor TPC/TPS Windows/Mac/Linux/WINGDK

diff --git a/Engine/Source/Developer/CookedEditor/CookedEditor.Build.cs b/Engine/Source/Developer/CookedEditor/CookedEditor.Build.cs
--- a/Engine/Source/Developer/CookedEditor/CookedEditor.Build.cs
+++ b/Engine/Source/Developer/CookedEditor/CookedEditor.Build.cs
@@ -27,12 +27,14 @@ public class CookedEditor : ModuleRules
 				"TargetPlatform",
 			});
 
-		PublicIncludePathModuleNames.Add("WindowsTargetPlatform");
+		PublicIncludePathModuleNames.Add("WindowsTargetPlatformSettings");
+		PublicIncludePathModuleNames.Add("WindowsTargetPlatformControls");
 
 		if (IsPlatformAvailable(UnrealTargetPlatform.Linux))
 		{
 			PublicDefinitions.Add("COOKEDEDITOR_WITH_LINUXTARGETPLATFORM=1");
-			PublicIncludePathModuleNames.Add("LinuxTargetPlatform");
+			PublicIncludePathModuleNames.Add("LinuxTargetPlatformSettings");
+			PublicIncludePathModuleNames.Add("LinuxTargetPlatformControls");
 		}
 		else
 		{
@@ -42,7 +44,8 @@ public class CookedEditor : ModuleRules
 		if (IsPlatformAvailable(UnrealTargetPlatform.Mac))
 		{
 			PublicDefinitions.Add("COOKEDEDITOR_WITH_MACTARGETPLATFORM=1");
-			PublicIncludePathModuleNames.Add("MacTargetPlatform");
+			PublicIncludePathModuleNames.Add("MacTargetPlatformSettings");
+			PublicIncludePathModuleNames.Add("MacTargetPlatformControls");
 		}
 		else
 		{
diff --git a/Engine/Source/Developer/CookedEditor/Private/CookedEditorPackageManager.cpp b/Engine/Source/Developer/CookedEditor/Private/CookedEditorPackageManager.cpp
--- a/Engine/Source/Developer/CookedEditor/Private/CookedEditorPackageManager.cpp
+++ b/Engine/Source/Developer/CookedEditor/Private/CookedEditorPackageManager.cpp
@@ -15,7 +15,7 @@ DEFINE_LOG_CATEGORY_STATIC(LogCookedEditorTargetPlatform, Log, All)
 
 
 
-TUniquePtr<ICookedEditorPackageManager> ICookedEditorPackageManager::FactoryForTargetPlatform(ITargetPlatform* TP, bool bIsCookedCooker)
+TUniquePtr<ICookedEditorPackageManager> ICookedEditorPackageManager::FactoryForTargetPlatform(bool bIsCookedCooker)
 {
 	if (FGameDelegates::Get().GetCookedEditorPackageManagerFactoryDelegate().IsBound())
 	{
@@ -51,7 +51,7 @@ void ICookedEditorPackageManager::AddPackagesFromPath(TArray<FName>& Packages, c
 	}
 }
 
-void ICookedEditorPackageManager::GatherAllPackagesExceptDisabled(TArray<FName>& PackageNames, const ITargetPlatform* TargetPlatform, const TArray<FString>& DisabledPlugins) const
+void ICookedEditorPackageManager::GatherAllPackagesExceptDisabled(TArray<FName>& PackageNames, const TArray<FString>& DisabledPlugins) const
 {
 	GetEnginePackagesToCook(PackageNames);
 	GetProjectPackagesToCook(PackageNames);
@@ -132,9 +132,9 @@ TArray<FString> FIniCookedEditorPackageManager::GetConfigArray(const TCHAR* Key)
 	return ResultArray;
 }
 
-void FIniCookedEditorPackageManager::GatherAllPackages(TArray<FName>& PackageNames, const ITargetPlatform* TargetPlatform) const
+void FIniCookedEditorPackageManager::GatherAllPackages(TArray<FName>& PackageNames) const
 {
-	GatherAllPackagesExceptDisabled(PackageNames, TargetPlatform, DisabledPlugins);
+	GatherAllPackagesExceptDisabled(PackageNames, DisabledPlugins);
 }
 
 void FIniCookedEditorPackageManager::FilterGatheredPackages(TArray<FName>& PackageNames) const
diff --git a/Engine/Source/Developer/CookedEditor/Private/CookedEditorTargetPlatform.cpp b/Engine/Source/Developer/CookedEditor/Private/CookedEditorTargetPlatform.cpp
--- a/Engine/Source/Developer/CookedEditor/Private/CookedEditorTargetPlatform.cpp
+++ b/Engine/Source/Developer/CookedEditor/Private/CookedEditorTargetPlatform.cpp
@@ -1,3 +1,4 @@
 // Copyright Epic Games, Inc. All Rights Reserved.
 
-#include "CookedEditorTargetPlatform.h"
+#include "CookedEditorTargetPlatformControls.h"
+#include "CookedEditorTargetPlatformSettings.h"
diff --git a/Engine/Source/Developer/CookedEditor/Public/CookedEditorPackageManager.h b/Engine/Source/Developer/CookedEditor/Public/CookedEditorPackageManager.h
--- a/Engine/Source/Developer/CookedEditor/Public/CookedEditorPackageManager.h
+++ b/Engine/Source/Developer/CookedEditor/Public/CookedEditorPackageManager.h
@@ -6,37 +6,43 @@
 #include "Common/TargetPlatformBase.h"
 #include "Interfaces/IPluginManager.h"
 
-#include "GenericWindowsTargetPlatform.h"
+#include "GenericWindowsTargetPlatformControls.h"
 #include "Windows/WindowsPlatformProperties.h"
 
 #if COOKEDEDITOR_WITH_LINUXTARGETPLATFORM
-	#include "LinuxTargetPlatform.h"
+	#include "LinuxTargetPlatformSettings.h"
+	#include "LinuxTargetPlatformControls.h"
 #endif
 
 #if COOKEDEDITOR_WITH_MACTARGETPLATFORM
-	#include "GenericMacTargetPlatform.h"
+	#include "GenericMacTargetPlatformSettings.h"
+	#include "GenericMacTargetPlatformControls.h"
 #endif
 
-typedef TGenericWindowsTargetPlatform<FWindowsPlatformProperties<false, false, false>> FWindowsEditorTargetPlatformParent;
+typedef TGenericWindowsTargetPlatformSettings<FWindowsPlatformProperties<false, false, false>> FWindowsEditorTargetPlatformSettingsParent;
+typedef TGenericWindowsTargetPlatformControls<FWindowsPlatformProperties<false, false, false>> FWindowsEditorTargetPlatformControlsParent;
 
 #if COOKEDEDITOR_WITH_LINUXTARGETPLATFORM
-typedef TLinuxTargetPlatform<FLinuxPlatformProperties<false, false, false, false>> FLinuxEditorTargetPlatformParent;
+typedef TLinuxTargetPlatformSettings<FLinuxPlatformProperties<false, false, false, false>> FLinuxEditorTargetPlatformSettingsParent;
+typedef TLinuxTargetPlatformControls<FLinuxPlatformProperties<false, false, false, false>> FLinuxEditorTargetPlatformControlsParent;
 #endif
 
 #if COOKEDEDITOR_WITH_MACTARGETPLATFORM
-typedef TGenericMacTargetPlatform<false, false, false> FMacEditorTargetPlatformParent;
+typedef TGenericMacTargetPlatformSettings<false, false, false> FMacEditorTargetPlatformSettingsParent;
+typedef TGenericMacTargetPlatformControls<false, false, false> FMacEditorTargetPlatformControlsParent;
 #endif
 
 #if PLATFORM_WINDOWS
-typedef FWindowsEditorTargetPlatformParent FHostPlatformEditorTargetPlatformParent;
+typedef FWindowsEditorTargetPlatformSettingsParent FHostPlatformEditorTargetPlatformSettingsParent;
+typedef FWindowsEditorTargetPlatformControlsParent FHostPlatformEditorTargetPlatformControlsParent;
 #elif PLATFORM_LINUX
-typedef FLinuxEditorTargetPlatformParent FHostPlatformEditorTargetPlatformParent;
+typedef FLinuxEditorTargetPlatformSettingsParent FHostPlatformEditorTargetPlatformSettingsParent;
+typedef FLinuxEditorTargetPlatformControlsParent FHostPlatformEditorTargetPlatformControlsParent;
 #elif PLATFORM_MAC
-typedef FMacEditorTargetPlatformParent FHostPlatformEditorTargetPlatformParent;
+typedef FMacEditorTargetPlatformSettingsParent FHostPlatformEditorTargetPlatformSettingsParent;
+typedef FMacEditorTargetPlatformControlsParent FHostPlatformEditorTargetPlatformControlsParent;
 #endif
 
-
-
 /**
  * Allows a project to control how packages are cooked when making a cooked editor.
  * 
@@ -64,7 +70,7 @@ public:
 	/** 
 	 * Construct a package manager for the given TP
 	 */
-	static TUniquePtr<ICookedEditorPackageManager> FactoryForTargetPlatform(ITargetPlatform* TP, bool bIsCookedCooker);
+	static TUniquePtr<ICookedEditorPackageManager> FactoryForTargetPlatform(bool bIsCookedCooker);
 
 	virtual ~ICookedEditorPackageManager()
 	{
@@ -106,7 +112,7 @@ public:
 	/**
 	 * Gathers the packages this PackageManager wants to manage (ie cook)
 	 */
-	virtual void GatherAllPackages(TArray<FName>& PackageNames, const ITargetPlatform* TargetPlatform) const = 0;
+	virtual void GatherAllPackages(TArray<FName>& PackageNames) const = 0;
 
 protected:
 
@@ -134,7 +140,7 @@ protected:
 	/**
 	 * Meat of this class, this calls other functions that generally will be overridden - subclass needs to pass in disabled plugins
 	 */
-	void GatherAllPackagesExceptDisabled(TArray<FName>& PackageNames, const ITargetPlatform* TargetPlatform, const TArray<FString>& DisabledPlugins) const;
+	void GatherAllPackagesExceptDisabled(TArray<FName>& PackageNames, const TArray<FString>& DisabledPlugins) const;
 };
 
 
@@ -169,5 +175,5 @@ public:
 	virtual bool AllowEnginePluginContentToBeCooked(const TSharedRef<IPlugin>) const override;
 	virtual bool AllowProjectPluginContentToBeCooked(const TSharedRef<IPlugin>) const override;
 	
-	virtual void GatherAllPackages(TArray<FName>& PackageNames, const ITargetPlatform* TargetPlatform) const override;
+	virtual void GatherAllPackages(TArray<FName>& PackageNames) const override;
 };
diff --git a/Engine/Source/Developer/CookedEditor/Public/CookedEditorTargetPlatform.h b/Engine/Source/Developer/CookedEditor/Public/CookedEditorTargetPlatform.h
deleted file mode 100644
diff --git a/Engine/Source/Developer/CookedEditor/Public/CookedEditorTargetPlatformControls.h b/Engine/Source/Developer/CookedEditor/Public/CookedEditorTargetPlatformControls.h
new file mode 100644
--- /dev/null
+++ b/Engine/Source/Developer/CookedEditor/Public/CookedEditorTargetPlatformControls.h
@@ -0,0 +1,220 @@
+// Copyright Epic Games, Inc. All Rights Reserved.
+
+#pragma once
+
+#include "Common/TargetPlatformControlsBase.h"
+#include "Interfaces/IPluginManager.h"
+#include "CookedEditorPackageManager.h"
+
+/**
+ * Templated class for a target platform used to cook a cooked editor. It needs to inherit from a 
+ * desktop platform's target platform such as TGenericWindowsTargetPlatformControls. See
+ * CookedEditorPackageManager.h for some typedef'd standard base classes that can be used.
+ * 
+ * The majority of the functionality is performed in an instance of a ICookedEditorPackageManager subclass.
+ * See that class for more information.
+ */
+template<typename Base>
+class TCookedEditorTargetPlatformControls : public Base
+{
+public:
+
+	TCookedEditorTargetPlatformControls(ITargetPlatformSettings* TargetPlatformSettings)
+		:Base(TargetPlatformSettings)
+	{
+		PackageManager = ICookedEditorPackageManager::FactoryForTargetPlatform(false);
+	}
+
+	/** Allows for a custom target platform module to initialize this TargetPlatform with an existing
+	 * PackageManager instead of going through the standard factory function */
+	explicit TCookedEditorTargetPlatformControls(TUniquePtr<ICookedEditorPackageManager>&& ExistingManager, ITargetPlatformSettings* TargetPlatformSettings)
+		:Base(TargetPlatformSettings)
+	{
+		PackageManager = MoveTemp(ExistingManager);
+	}
+
+	~TCookedEditorTargetPlatformControls()
+	{
+	}
+
+	virtual FString PlatformName() const override
+	{
+		static FString CachedPlatformName = this->GetTargetPlatformSettings()->IniPlatformName() + TEXT("CookedEditor");
+		return CachedPlatformName;
+	}
+
+	virtual FString CookingDeviceProfileName() const override
+	{
+		return Base::PlatformName();
+	}
+
+	/**
+	 * If you override this to return false, you will have to stage uncooked assets to allow the editor to run properly
+	 */
+	virtual bool AllowsEditorObjects() const override
+	{
+		return true;
+	}
+
+	virtual bool AllowObject(const class UObject* Obj) const override
+	{
+		// probably don't need this check, but it can't hurt
+		if (!AllowsEditorObjects())
+		{
+			return true;
+		}
+
+		return PackageManager->AllowObjectToBeCooked(Obj);
+	}
+
+	virtual void GetExtraPackagesToCook(TArray<FName>& PackageNames) const override
+	{
+		if (AllowsEditorObjects())
+		{
+			PackageManager->GatherAllPackages(PackageNames);
+		}
+	}
+
+	virtual bool IsRunningPlatform() const override
+	{
+		return false;
+	}
+
+	virtual float GetVariantPriority() const override
+	{
+		// by returning -1, we will never use this variant when targeting host platform this class implements
+		// (without this, cooking for Windows in the editor may choose this variant to cook for, which we never want)
+		return -1.0f;
+	}
+
+	virtual void GetAllDevices(TArray<ITargetDevicePtr>& OutDevices) const override
+	{
+	}
+
+	virtual ITargetDevicePtr GetDefaultDevice() const override
+	{
+		return nullptr;
+	}
+
+
+protected:
+
+	TUniquePtr<ICookedEditorPackageManager> PackageManager;
+
+};
+
+
+template<typename Base>
+class TCookedCookerTargetPlatformControls : public Base
+{
+public:
+
+	TCookedCookerTargetPlatformControls(ITargetPlatformSettings* TargetPlatformSettings)
+		: Base(TargetPlatformSettings)
+	{
+		PackageManager = ICookedEditorPackageManager::FactoryForTargetPlatform(true);
+	}
+
+	TCookedCookerTargetPlatformControls(TUniquePtr<ICookedEditorPackageManager>&& ExistingManager, ITargetPlatformSettings* TargetPlatformSettings)
+		: Base(TargetPlatformSettings)
+	{
+		PackageManager = MoveTemp(ExistingManager);
+	}
+
+	~TCookedCookerTargetPlatformControls()
+	{
+	}
+
+
+	virtual FString PlatformName() const override
+	{
+		static FString CachedPlatformName = this->GetTargetPlatformSettings()->IniPlatformName() + TEXT("CookedCooker");
+		return CachedPlatformName;
+	}
+
+	virtual FString CookingDeviceProfileName() const override
+	{
+		return Base::PlatformName();
+	}
+
+	virtual bool AllowsEditorObjects() const override
+	{
+		return true;
+	}
+
+	virtual bool AllowsDevelopmentObjects() const override
+	{
+		return false;
+	}
+
+	virtual bool AllowObject(const UObject* Obj) const override
+	{
+		// probably don't need this check, but it can't hurt
+		if (!AllowsEditorObjects())
+		{
+			return true;
+		}
+
+		return PackageManager->AllowObjectToBeCooked(Obj);
+	}
+	
+	virtual void GetExtraPackagesToCook(TArray<FName>& PackageNames) const override
+	{
+		if (AllowsEditorObjects())
+		{
+			PackageManager->GatherAllPackages(PackageNames);
+		}
+	}
+
+
+	//////////////////////////////////////////////////////////////////
+	// Disabling stuff since it's just a cooker
+	//////////////////////////////////////////////////////////////////
+
+
+	virtual void GetTextureFormats(const class UTexture* InTexture, TArray< TArray<FName> >& OutFormats) const override
+	{
+		// no textures please
+	}
+
+	virtual void GetAllTextureFormats(TArray<FName>& OutFormats) const override
+	{
+		// no textures please
+	}
+
+	virtual bool SupportsVariants() const override
+	{
+		return false;
+	}
+
+	virtual bool AllowAudioVisualData() const override
+	{
+		return false;
+	}
+
+	virtual bool IsRunningPlatform() const override
+	{
+		return false;
+	}
+
+	virtual float GetVariantPriority() const override
+	{
+		// by returning -1, we will never use this variant when targeting host platform this class implements
+		// (without this, cooking for Windows in the editor may choose this variant to cook for, which we never want)
+		return -1.0f;
+	}
+
+	virtual void GetAllDevices(TArray<ITargetDevicePtr>& OutDevices) const override
+	{
+	}
+
+	virtual ITargetDevicePtr GetDefaultDevice() const override
+	{
+		return nullptr;
+	}
+
+
+
+	TUniquePtr<ICookedEditorPackageManager> PackageManager;
+
+};
diff --git a/Engine/Source/Developer/CookedEditor/Public/CookedEditorTargetPlatformSettings.h b/Engine/Source/Developer/CookedEditor/Public/CookedEditorTargetPlatformSettings.h
new file mode 100644
--- /dev/null
+++ b/Engine/Source/Developer/CookedEditor/Public/CookedEditorTargetPlatformSettings.h
@@ -0,0 +1,108 @@
+// Copyright Epic Games, Inc. All Rights Reserved.
+
+#pragma once
+
+#include "Common/TargetPlatformSettingsBase.h"
+#include "Interfaces/IPluginManager.h"
+#include "CookedEditorPackageManager.h"
+
+/**
+ * Templated class for a target platform used to cook a cooked editor. It needs to inherit from a 
+ * desktop platform's target platform such as TGenericWindowsTargetPlatformSettings. See
+ * CookedEditorPackageManager.h for some typedef'd standard base classes that can be used.
+ * 
+ * The majority of the functionality is performed in an instance of a ICookedEditorPackageManager subclass.
+ * See that class for more information.
+ */
+template<typename Base>
+class TCookedEditorTargetPlatformSettings : public Base
+{
+public:
+
+	TCookedEditorTargetPlatformSettings()
+	{
+		PackageManager = ICookedEditorPackageManager::FactoryForTargetPlatform(false);
+	}
+
+	/** Allows for a custom target platform module to initialize this TargetPlatform with an existing
+	 * PackageManager instead of going through the standard factory function */
+	explicit TCookedEditorTargetPlatformSettings(TUniquePtr<ICookedEditorPackageManager>&& ExistingManager)
+	{
+		PackageManager = MoveTemp(ExistingManager);
+	}
+
+	~TCookedEditorTargetPlatformSettings()
+	{
+	}
+
+	virtual void GetReflectionCaptureFormats(TArray<FName>& OutFormats) const override
+	{
+		Base::GetReflectionCaptureFormats(OutFormats);
+
+		// UMapBuildDataRegistry::PostLoad() assumes that editor always needs encoded data, so when cooking for the editor,
+		// make sure that the EncodedHDR format is included
+		OutFormats.AddUnique(TEXT("EncodedHDR"));
+	}
+
+
+
+protected:
+
+	TUniquePtr<ICookedEditorPackageManager> PackageManager;
+
+};
+
+
+template<typename Base>
+class TCookedCookerTargetPlatformSettings : public Base
+{
+public:
+
+	TCookedCookerTargetPlatformSettings()
+	{
+		PackageManager = ICookedEditorPackageManager::FactoryForTargetPlatform(true);
+	}
+
+	TCookedCookerTargetPlatformSettings(TUniquePtr<ICookedEditorPackageManager>&& ExistingManager)
+	{
+		PackageManager = MoveTemp(ExistingManager);
+	}
+
+	~TCookedCookerTargetPlatformSettings()
+	{
+	}
+	
+	//////////////////////////////////////////////////////////////////
+	// Disabling stuff since it's just a cooker
+	//////////////////////////////////////////////////////////////////
+
+	virtual void GetAllPossibleShaderFormats(TArray<FName>& OutFormats) const override
+	{
+		// no shaders please
+	}
+
+	virtual void GetAllTargetedShaderFormats(TArray<FName>& OutFormats) const override
+	{
+		// no shaders please
+	}
+
+	virtual bool SupportsFeature(ETargetPlatformFeatures Feature) const override
+	{
+		switch (Feature)
+		{
+		case ETargetPlatformFeatures::AudioStreaming:
+		case ETargetPlatformFeatures::MemoryMappedAudio:
+			return false;
+
+		case ETargetPlatformFeatures::CanCookPackages:
+		case ETargetPlatformFeatures::Packaging:
+			return true;
+
+		default:
+			return Base::SupportsFeature(Feature);
+		}
+	}
+
+	TUniquePtr<ICookedEditorPackageManager> PackageManager;
+
+};
diff --git a/Engine/Source/Developer/DesktopPlatform/Public/SteamDeck/SteamDeckDevice.h b/Engine/Source/Developer/DesktopPlatform/Public/SteamDeck/SteamDeckDevice.h
--- a/Engine/Source/Developer/DesktopPlatform/Public/SteamDeck/SteamDeckDevice.h
+++ b/Engine/Source/Developer/DesktopPlatform/Public/SteamDeck/SteamDeckDevice.h
@@ -15,8 +15,8 @@ template<class ParentDeviceClass>
 class TSteamDeckDevice : public ParentDeviceClass
 {
 public:
-	TSteamDeckDevice(FString InIpAddr, FString InDeviceName, FString InUserName, FString InPassword, const ITargetPlatform& InTargetPlatform, const TCHAR* InRuntimeOSName)
-		: ParentDeviceClass(InTargetPlatform)
+	TSteamDeckDevice(FString InIpAddr, FString InDeviceName, FString InUserName, FString InPassword, const ITargetPlatformControls& InTargetPlatformControls, const TCHAR* InRuntimeOSName)
+		: ParentDeviceClass(InTargetPlatformControls)
 		, IpAddr(InIpAddr)
 		, UserName(InUserName)
 		, Password(InPassword)
@@ -32,7 +32,7 @@ public:
 
 	virtual FTargetDeviceId GetId() const override
 	{
-		return FTargetDeviceId(this->TargetPlatform.PlatformName(), IpAddr);
+		return FTargetDeviceId(this->TargetPlatformControls.PlatformName(), IpAddr);
 	}
 
 	virtual FString GetOperatingSystemName() override
@@ -48,7 +48,7 @@ public:
 		return true;
 	}
 
-	static TArray<ITargetDevicePtr> DiscoverDevices(const ITargetPlatform& TargetPlatform, const TCHAR* RuntimeOSName)
+	static TArray<ITargetDevicePtr> DiscoverDevices(const ITargetPlatformControls& TargetPlatformControls, const TCHAR* RuntimeOSName)
 	{
 		TArray<FString> EngineIniSteamDeckDevices;
 
@@ -92,7 +92,7 @@ public:
 				ConfigPassword = FString();
 			}
 
-			SteamDevices.Add(MakeShareable(new TSteamDeckDevice<ParentDeviceClass>(IpAddr, Name, ConfigUserName, ConfigPassword, TargetPlatform, RuntimeOSName)));
+			SteamDevices.Add(MakeShareable(new TSteamDeckDevice<ParentDeviceClass>(IpAddr, Name, ConfigUserName, ConfigPassword, TargetPlatformControls, RuntimeOSName)));
 		}
 
 		return SteamDevices;
diff --git a/Engine/Source/Developer/IOS/IOSPlatformEditor/IOSPlatformEditor.Build.cs b/Engine/Source/Developer/IOS/IOSPlatformEditor/IOSPlatformEditor.Build.cs
--- a/Engine/Source/Developer/IOS/IOSPlatformEditor/IOSPlatformEditor.Build.cs
+++ b/Engine/Source/Developer/IOS/IOSPlatformEditor/IOSPlatformEditor.Build.cs
@@ -29,7 +29,7 @@ public class IOSPlatformEditor : ModuleRules
                 "AudioSettingsEditor",
                 "GameProjectGeneration",
 				"FreeImage",
-                "MacTargetPlatform",
+                "MacTargetPlatformSettings",
 			}
 		);
 
diff --git a/Engine/Source/Developer/Linux/LinuxArm64TargetPlatform/LinuxArm64TargetPlatform.Build.cs b/Engine/Source/Developer/Linux/LinuxArm64TargetPlatform/LinuxArm64TargetPlatform.Build.cs
--- a/Engine/Source/Developer/Linux/LinuxArm64TargetPlatform/LinuxArm64TargetPlatform.Build.cs
+++ b/Engine/Source/Developer/Linux/LinuxArm64TargetPlatform/LinuxArm64TargetPlatform.Build.cs
@@ -15,17 +15,5 @@ public class LinuxArm64TargetPlatform : ModuleRules
 				"TargetPlatform",
 			}
         );
-
-        if (Target.bCompileAgainstEngine)
-        {
-            PrivateDependencyModuleNames.AddRange(new string[] {
-					"Engine"
-				}
-            );
-
-            PrivateIncludePathModuleNames.Add("TextureCompressor");
-        }
-
-		PrivateIncludePathModuleNames.Add("LinuxTargetPlatform");
     }
 }
diff --git a/Engine/Source/Developer/Linux/LinuxArm64TargetPlatform/Private/LinuxArm64TargetPlatformModule.cpp b/Engine/Source/Developer/Linux/LinuxArm64TargetPlatform/Private/LinuxArm64TargetPlatformModule.cpp
--- a/Engine/Source/Developer/Linux/LinuxArm64TargetPlatform/Private/LinuxArm64TargetPlatformModule.cpp
+++ b/Engine/Source/Developer/Linux/LinuxArm64TargetPlatform/Private/LinuxArm64TargetPlatformModule.cpp
@@ -6,12 +6,8 @@
 
 #include "CoreMinimal.h"
 #include "Modules/ModuleManager.h"
-
 #include "Interfaces/ITargetPlatformModule.h"
-
-#include "LinuxTargetDevice.h"
-#include "LinuxTargetPlatform.h"
-
+#include "Common/TargetPlatformBase.h"
 
 /**
  * Module for the Linux target platforms
@@ -23,12 +19,14 @@ public:
 
 	virtual void GetTargetPlatforms(TArray<ITargetPlatform*>& TargetPlatforms) override
 	{
-		// Game TP
-		TargetPlatforms.Add(new TLinuxTargetPlatform<FLinuxPlatformProperties<false, false, false, true> >());
-		// Server TP
-		TargetPlatforms.Add(new TLinuxTargetPlatform<FLinuxPlatformProperties<false, true, false, true> >());
-		// Client TP
-		TargetPlatforms.Add(new TLinuxTargetPlatform<FLinuxPlatformProperties<false, false, true, true> >());
+	}
+
+	virtual void GetTargetPlatforms(TArray<ITargetPlatform*>& TargetPlatforms, TArray<ITargetPlatformSettings*> TargetPlatformSettings, TArray<ITargetPlatformControls*> TargetPlatformControls) override
+	{
+		for (ITargetPlatformControls* TargetPlatformControlsIt : TargetPlatformControls)
+		{
+			TargetPlatforms.Add(new FTargetPlatformMerged(TargetPlatformControlsIt->GetTargetPlatformSettings(), TargetPlatformControlsIt));
+		}
 	}
 };
 
diff --git a/Engine/Source/Developer/Linux/LinuxArm64TargetPlatformControls/LinuxArm64TargetPlatformControls.Build.cs b/Engine/Source/Developer/Linux/LinuxArm64TargetPlatformControls/LinuxArm64TargetPlatformControls.Build.cs
new file mode 100644
--- /dev/null
+++ b/Engine/Source/Developer/Linux/LinuxArm64TargetPlatformControls/LinuxArm64TargetPlatformControls.Build.cs
@@ -0,0 +1,39 @@
+// Copyright Epic Games, Inc. All Rights Reserved.
+
+using UnrealBuildTool;
+
+public class LinuxArm64TargetPlatformControls : ModuleRules
+{
+    public LinuxArm64TargetPlatformControls(ReadOnlyTargetRules Target) : base(Target)
+    {
+        BinariesSubFolder = "LinuxArm64";
+		// We need a short name here since this can run afoul very easily of the `MAX_PATH` limit when
+		// combined with building other targets that make use of this.
+		ShortName = "LArm64TPCon";
+
+		PrivateDependencyModuleNames.AddRange(
+            new string[] {
+				"Core",
+				"DesktopPlatform",
+				"TargetPlatform",
+				"LinuxArm64TargetPlatformSettings",
+				"LinuxTargetPlatformSettings",
+				"LinuxTargetPlatformControls",
+			}
+        );
+
+        if (Target.bCompileAgainstEngine)
+        {
+            PrivateDependencyModuleNames.AddRange(new string[] {
+					"Engine"
+				}
+            );
+
+            PrivateIncludePathModuleNames.Add("TextureCompressor");
+        }
+
+		PrivateIncludePathModuleNames.Add("LinuxArm64TargetPlatformSettings");
+		PrivateIncludePathModuleNames.Add("LinuxTargetPlatformSettings");
+		PrivateIncludePathModuleNames.Add("LinuxTargetPlatformControls");
+    }
+}
diff --git a/Engine/Source/Developer/Linux/LinuxArm64TargetPlatformControls/Private/LinuxArm64TargetPlatformControlsModule.cpp b/Engine/Source/Developer/Linux/LinuxArm64TargetPlatformControls/Private/LinuxArm64TargetPlatformControlsModule.cpp
new file mode 100644
--- /dev/null
+++ b/Engine/Source/Developer/Linux/LinuxArm64TargetPlatformControls/Private/LinuxArm64TargetPlatformControlsModule.cpp
@@ -0,0 +1,41 @@
+// Copyright Epic Games, Inc. All Rights Reserved.
+
+/*=============================================================================
+	FLinuxArm64TargetPlatformControlsModule.cpp: Implements the FLinuxArm64TargetPlatformControlsModule class.
+=============================================================================*/
+
+#include "CoreMinimal.h"
+#include "Modules/ModuleManager.h"
+
+#include "Interfaces/ITargetPlatformControlsModule.h"
+
+#include "LinuxTargetPlatformSettings.h"
+#include "LinuxTargetPlatformControls.h"
+#include "ILinuxArm64TargetPlatformSettingsModule.h"
+
+/**
+ * Module for the Linux target platforms controls
+ */
+class FLinuxArm64TargetPlatformControlsModule
+	: public ITargetPlatformControlsModule
+{
+public:
+	virtual void GetTargetPlatformControls(TArray<ITargetPlatformControls*>& TargetPlatforms, FName& PlatformSettingsModuleName) override
+	{
+		TMap<FString, ITargetPlatformSettings*> PlatformNameToPlatformSettings;
+		ILinuxArm64TargetPlatformSettingsModule* ModuleSettings = FModuleManager::GetModulePtr<ILinuxArm64TargetPlatformSettingsModule>(PlatformSettingsModuleName);
+		if (ModuleSettings != nullptr)
+		{
+			TMap<FString, ITargetPlatformSettings*> OutMap;
+			ModuleSettings->GetPlatformSettingsMaps(OutMap);
+			ITargetPlatformControls* GameTP = new TLinuxTargetPlatformControls<FLinuxPlatformProperties<false, false, false, true>>(OutMap[FLinuxPlatformProperties<false, false, false, true>::PlatformName()]);
+			ITargetPlatformControls* ServerTP = new TLinuxTargetPlatformControls<FLinuxPlatformProperties<false, true, false, true>>(OutMap[FLinuxPlatformProperties<false, true, false, true>::PlatformName()]);
+			ITargetPlatformControls* ClientTP = new TLinuxTargetPlatformControls<FLinuxPlatformProperties<false, false, true, true>>(OutMap[FLinuxPlatformProperties<false, false, true, true>::PlatformName()]);
+			TargetPlatforms.Add(GameTP);
+			TargetPlatforms.Add(ServerTP);
+			TargetPlatforms.Add(ClientTP);
+		}
+	}
+};
+
+IMPLEMENT_MODULE(FLinuxArm64TargetPlatformControlsModule, LinuxArm64TargetPlatformControls);
diff --git a/Engine/Source/Developer/Linux/LinuxArm64TargetPlatformSettings/LinuxArm64TargetPlatformSettings.Build.cs b/Engine/Source/Developer/Linux/LinuxArm64TargetPlatformSettings/LinuxArm64TargetPlatformSettings.Build.cs
new file mode 100644
--- /dev/null
+++ b/Engine/Source/Developer/Linux/LinuxArm64TargetPlatformSettings/LinuxArm64TargetPlatformSettings.Build.cs
@@ -0,0 +1,35 @@
+// Copyright Epic Games, Inc. All Rights Reserved.
+
+using UnrealBuildTool;
+
+public class LinuxArm64TargetPlatformSettings : ModuleRules
+{
+    public LinuxArm64TargetPlatformSettings(ReadOnlyTargetRules Target) : base(Target)
+    {
+        BinariesSubFolder = "LinuxArm64";
+		// We need a short name here since this can run afoul very easily of the `MAX_PATH` limit when
+		// combined with building other targets that make use of this.
+		ShortName = "LArm64TPSet";
+
+		PrivateDependencyModuleNames.AddRange(
+            new string[] {
+				"Core",
+				"DesktopPlatform",
+				"TargetPlatform",
+				"LinuxTargetPlatformSettings",
+			}
+        );
+
+        if (Target.bCompileAgainstEngine)
+        {
+            PrivateDependencyModuleNames.AddRange(new string[] {
+					"Engine"
+				}
+            );
+
+            PrivateIncludePathModuleNames.Add("TextureCompressor");
+        }
+
+		PrivateIncludePathModuleNames.Add("LinuxTargetPlatformSettings");
+	}
+}
diff --git a/Engine/Source/Developer/Linux/LinuxArm64TargetPlatformSettings/Private/LinuxArm64TargetPlatformSettingsModule.cpp b/Engine/Source/Developer/Linux/LinuxArm64TargetPlatformSettings/Private/LinuxArm64TargetPlatformSettingsModule.cpp
new file mode 100644
--- /dev/null
+++ b/Engine/Source/Developer/Linux/LinuxArm64TargetPlatformSettings/Private/LinuxArm64TargetPlatformSettingsModule.cpp
@@ -0,0 +1,52 @@
+// Copyright Epic Games, Inc. All Rights Reserved.
+
+/*=============================================================================
+	LinuxArm64TargetPlatformSettingsModule.cpp: Implements the FLinuxArm64TargetPlatformSettingsModule class.
+=============================================================================*/
+
+#include "Modules/ModuleManager.h"
+#include "ILinuxArm64TargetPlatformSettingsModule.h"
+#include "LinuxTargetPlatformSettings.h"
+
+
+/**
+ * Module for the Linux target platforms settings
+ */
+class FLinuxArm64TargetPlatformSettingsModule
+	: public ILinuxArm64TargetPlatformSettingsModule
+{
+public:
+
+	/** Destructor. */
+	~FLinuxArm64TargetPlatformSettingsModule()
+	{
+
+	}
+
+	virtual void GetTargetPlatformSettings(TArray<ITargetPlatformSettings*>& TargetPlatforms) override
+	{
+		// NoEditor TP
+		ITargetPlatformSettings* NoEditorTP = new TLinuxTargetPlatformSettings<FLinuxPlatformProperties<false, false, false, true>>();
+		TargetPlatforms.Add(NoEditorTP);
+		PlatformNameToPlatformSettings.Add(FLinuxPlatformProperties<false, false, false, true>::PlatformName(), NoEditorTP);
+		// Server TP
+		ITargetPlatformSettings* ServerTP = new TLinuxTargetPlatformSettings<FLinuxPlatformProperties<false, true, false, true>>();
+		TargetPlatforms.Add(ServerTP);
+		PlatformNameToPlatformSettings.Add(FLinuxPlatformProperties<false, true, false, true>::PlatformName(), ServerTP);
+		// Client TP
+		ITargetPlatformSettings* ClientTP = new TLinuxTargetPlatformSettings<FLinuxPlatformProperties<false, false, true, true>>();
+		TargetPlatforms.Add(ClientTP);
+		PlatformNameToPlatformSettings.Add(FLinuxPlatformProperties<false, false, true, true>::PlatformName(), ClientTP);
+	}
+
+	virtual void GetPlatformSettingsMaps(TMap<FString, ITargetPlatformSettings*>& OutMap) override
+	{
+		OutMap = PlatformNameToPlatformSettings;
+	}
+
+private:
+	TMap<FString, ITargetPlatformSettings*> PlatformNameToPlatformSettings;
+
+};
+
+IMPLEMENT_MODULE(FLinuxArm64TargetPlatformSettingsModule, LinuxArm64TargetPlatformSettings);
\ No newline at end of file
diff --git a/Engine/Source/Developer/Linux/LinuxArm64TargetPlatformSettings/Public/ILinuxArm64TargetPlatformSettingsModule.h b/Engine/Source/Developer/Linux/LinuxArm64TargetPlatformSettings/Public/ILinuxArm64TargetPlatformSettingsModule.h
new file mode 100644
--- /dev/null
+++ b/Engine/Source/Developer/Linux/LinuxArm64TargetPlatformSettings/Public/ILinuxArm64TargetPlatformSettingsModule.h
@@ -0,0 +1,16 @@
+// Copyright Epic Games, Inc. All Rights Reserved.
+
+#pragma  once 
+
+/*------------------------------------------------------------------------------------
+IWindowsTargetPlatformSettingsModule interface
+------------------------------------------------------------------------------------*/
+
+#include "CoreMinimal.h"
+#include "Interfaces/ITargetPlatformSettingsModule.h"
+
+class ILinuxArm64TargetPlatformSettingsModule : public ITargetPlatformSettingsModule
+{
+public:
+	virtual void GetPlatformSettingsMaps(TMap<FString, ITargetPlatformSettings*>& OutMap) = 0;
+};
diff --git a/Engine/Source/Developer/Linux/LinuxPlatformEditor/LinuxPlatformEditor.Build.cs b/Engine/Source/Developer/Linux/LinuxPlatformEditor/LinuxPlatformEditor.Build.cs
--- a/Engine/Source/Developer/Linux/LinuxPlatformEditor/LinuxPlatformEditor.Build.cs
+++ b/Engine/Source/Developer/Linux/LinuxPlatformEditor/LinuxPlatformEditor.Build.cs
@@ -20,7 +20,7 @@ public class LinuxPlatformEditor : ModuleRules
 				"PropertyEditor",
 				"SharedSettingsWidgets",
 				"SourceControl",
-				"LinuxTargetPlatform",
+				"LinuxTargetPlatformSettings",
 				"TargetPlatform",
 				"MaterialShaderQualitySettings",
 				"RenderCore",
diff --git a/Engine/Source/Developer/Linux/LinuxTargetPlatform/Classes/LinuxTargetSettings.h b/Engine/Source/Developer/Linux/LinuxTargetPlatform/Classes/LinuxTargetSettings.h
deleted file mode 100644
diff --git a/Engine/Source/Developer/Linux/LinuxTargetPlatform/LinuxTargetPlatform.Build.cs b/Engine/Source/Developer/Linux/LinuxTargetPlatform/LinuxTargetPlatform.Build.cs
--- a/Engine/Source/Developer/Linux/LinuxTargetPlatform/LinuxTargetPlatform.Build.cs
+++ b/Engine/Source/Developer/Linux/LinuxTargetPlatform/LinuxTargetPlatform.Build.cs
@@ -11,24 +11,9 @@ public class LinuxTargetPlatform : ModuleRules
 		PrivateDependencyModuleNames.AddRange(
             new string[] {
 				"Core",
-				"CoreUObject",
 				"TargetPlatform",
 				"DesktopPlatform",
-				"Projects"
 			}
         );
-
-		PrivateIncludePathModuleNames.AddRange(
-			new string[] {
-				"Settings"
-			}
-		);
-
-		if (Target.bCompileAgainstEngine)
-		{
-			PublicIncludePathModuleNames.Add("Engine");
-			PrivateDependencyModuleNames.Add("Engine");
-			PrivateIncludePathModuleNames.Add("TextureCompressor");
-		}
 	}
 }
diff --git a/Engine/Source/Developer/Linux/LinuxTargetPlatform/Private/LinuxTargetPlatformClasses.cpp b/Engine/Source/Developer/Linux/LinuxTargetPlatform/Private/LinuxTargetPlatformClasses.cpp
deleted file mode 100644
diff --git a/Engine/Source/Developer/Linux/LinuxTargetPlatform/Private/LinuxTargetPlatformModule.cpp b/Engine/Source/Developer/Linux/LinuxTargetPlatform/Private/LinuxTargetPlatformModule.cpp
--- a/Engine/Source/Developer/Linux/LinuxTargetPlatform/Private/LinuxTargetPlatformModule.cpp
+++ b/Engine/Source/Developer/Linux/LinuxTargetPlatform/Private/LinuxTargetPlatformModule.cpp
@@ -2,22 +2,11 @@
 
 #include "CoreMinimal.h"
 #include "Modules/ModuleManager.h"
-#include "UObject/Package.h"
-#include "UObject/WeakObjectPtr.h"
-
 #include "Interfaces/ITargetPlatformModule.h"
-
-#include "LinuxTargetSettings.h"
-#include "LinuxTargetDevice.h"
-#include "LinuxTargetPlatform.h"
-
-#include "ISettingsModule.h"
-
+#include "Common/TargetPlatformBase.h"
 
 #define LOCTEXT_NAMESPACE "FLinuxTargetPlatformModule"
 
-
-
 /**
  * Module for the Linux target platform.
  */
@@ -25,16 +14,17 @@ class FLinuxTargetPlatformModule
 	: public ITargetPlatformModule
 {
 public:
-	virtual void GetTargetPlatforms(TArray<ITargetPlatform*>& TargetPlatforms) override
+	virtual void GetTargetPlatforms(TArray<ITargetPlatform*>& TargetPlatforms)
 	{
-		// NoEditor TP
-		TargetPlatforms.Add(new TLinuxTargetPlatform<FLinuxPlatformProperties<false, false, false, false> >());
-		// Editor TP
-		TargetPlatforms.Add(new TLinuxTargetPlatform<FLinuxPlatformProperties<true, false, false, false> >());
-		// Server TP
-		TargetPlatforms.Add(new TLinuxTargetPlatform<FLinuxPlatformProperties<false, true, false, false> >());
-		// Client TP
-		TargetPlatforms.Add(new TLinuxTargetPlatform<FLinuxPlatformProperties<false, false, true, false> >());
+
+	}
+
+	virtual void GetTargetPlatforms(TArray<ITargetPlatform*>& TargetPlatforms, TArray<ITargetPlatformSettings*> TargetPlatformSettings, TArray<ITargetPlatformControls*> TargetPlatformControls) override
+	{
+		for (ITargetPlatformControls* TargetPlatformControlsIt : TargetPlatformControls)
+		{
+			TargetPlatforms.Add(new FTargetPlatformMerged(TargetPlatformControlsIt->GetTargetPlatformSettings(), TargetPlatformControlsIt));
+		}
 	}
 
 public:
@@ -43,48 +33,13 @@ public:
 
 	virtual void StartupModule() override
 	{
-		TargetSettings = NewObject<ULinuxTargetSettings>(GetTransientPackage(), "LinuxTargetSettings", RF_Standalone);
-
-		// We need to manually load the config properties here, as this module is loaded before the UObject system is setup to do this
-		GConfig->GetArray(TEXT("/Script/LinuxTargetPlatform.LinuxTargetSettings"), TEXT("TargetedRHIs"), TargetSettings->TargetedRHIs, GEngineIni);
-		TargetSettings->AddToRoot();
 
-		ISettingsModule* SettingsModule = FModuleManager::GetModulePtr<ISettingsModule>("Settings");
-
-		if (SettingsModule != nullptr)
-		{
-			SettingsModule->RegisterSettings("Project", "Platforms", "Linux",
-				LOCTEXT("TargetSettingsName", "Linux"),
-				LOCTEXT("TargetSettingsDescription", "Settings for Linux target platform"),
-				TargetSettings
-			);
-		}
 	}
 
 	virtual void ShutdownModule() override
 	{
-		ISettingsModule* SettingsModule = FModuleManager::GetModulePtr<ISettingsModule>("Settings");
-
-		if (SettingsModule != nullptr)
-		{
-			SettingsModule->UnregisterSettings("Project", "Platforms", "Linux");
-		}
 
-		if (!GExitPurge)
-		{
-			// If we're in exit purge, this object has already been destroyed
-			TargetSettings->RemoveFromRoot();
-		}
-		else
-		{
-			TargetSettings = NULL;
-		}
 	}
-
-private:
-
-	/** Holds the target settings. */
-	ULinuxTargetSettings* TargetSettings;
 };
 
 
diff --git a/Engine/Source/Developer/Linux/LinuxTargetPlatform/Public/LinuxTargetDevice.h b/Engine/Source/Developer/Linux/LinuxTargetPlatform/Public/LinuxTargetDevice.h
deleted file mode 100644
diff --git a/Engine/Source/Developer/Linux/LinuxTargetPlatform/Public/LinuxTargetPlatform.h b/Engine/Source/Developer/Linux/LinuxTargetPlatform/Public/LinuxTargetPlatform.h
deleted file mode 100644
diff --git a/Engine/Source/Developer/Linux/LinuxTargetPlatformControls/LinuxTargetPlatformControls.Build.cs b/Engine/Source/Developer/Linux/LinuxTargetPlatformControls/LinuxTargetPlatformControls.Build.cs
new file mode 100644
--- /dev/null
+++ b/Engine/Source/Developer/Linux/LinuxTargetPlatformControls/LinuxTargetPlatformControls.Build.cs
@@ -0,0 +1,38 @@
+// Copyright Epic Games, Inc. All Rights Reserved.
+
+using UnrealBuildTool;
+
+public class LinuxTargetPlatformControls : ModuleRules
+{
+    public LinuxTargetPlatformControls(ReadOnlyTargetRules Target) : base(Target)
+	{
+        BinariesSubFolder = "Linux";
+		// We need a short name here since this can run afoul very easily of the `MAX_PATH` limit when
+		// combined with building other targets that make use of this.
+		ShortName = "LinuxTPCon";
+
+		PrivateDependencyModuleNames.AddRange(
+            new string[] {
+				"Core",
+				"CoreUObject",
+				"TargetPlatform",
+				"DesktopPlatform",
+				"Projects",
+				"LinuxTargetPlatformSettings",
+			}
+        );
+
+		PrivateIncludePathModuleNames.AddRange(
+			new string[] {
+				"Settings"
+			}
+		);
+
+		if (Target.bCompileAgainstEngine)
+		{
+			PublicIncludePathModuleNames.Add("Engine");
+			PrivateDependencyModuleNames.Add("Engine");
+			PrivateIncludePathModuleNames.Add("TextureCompressor");
+		}
+	}
+}
diff --git a/Engine/Source/Developer/Linux/LinuxTargetPlatformControls/Private/LinuxTargetPlatformControlsModule.cpp b/Engine/Source/Developer/Linux/LinuxTargetPlatformControls/Private/LinuxTargetPlatformControlsModule.cpp
new file mode 100644
--- /dev/null
+++ b/Engine/Source/Developer/Linux/LinuxTargetPlatformControls/Private/LinuxTargetPlatformControlsModule.cpp
@@ -0,0 +1,58 @@
+// Copyright Epic Games, Inc. All Rights Reserved.
+
+#include "CoreMinimal.h"
+#include "Modules/ModuleManager.h"
+#include "Interfaces/ITargetPlatformControlsModule.h"
+#include "LinuxTargetPlatformControls.h"
+#include "ILinuxTargetPlatformSettingsModule.h"
+
+#define LOCTEXT_NAMESPACE "FLinuxTargetPlatformControlModule"
+
+
+
+/**
+ * Module for the Linux target platform controls.
+ */
+class FLinuxTargetPlatformControlsModule
+	: public ITargetPlatformControlsModule
+{
+public:
+
+	virtual void GetTargetPlatformControls(TArray<ITargetPlatformControls*>& TargetPlatforms, FName& PlatformSettingsModuleName) override
+	{
+		TMap<FString, ITargetPlatformSettings*> PlatformNameToPlatformSettings;
+		ILinuxTargetPlatformSettingsModule* ModuleSettings = FModuleManager::GetModulePtr<ILinuxTargetPlatformSettingsModule>(PlatformSettingsModuleName);
+		if (ModuleSettings != nullptr)
+		{
+			TMap<FString, ITargetPlatformSettings*> OutMap;
+			ModuleSettings->GetPlatformSettingsMaps(OutMap);
+			ITargetPlatformControls* GameTP = new TLinuxTargetPlatformControls<FLinuxPlatformProperties<false, false, false, false>>(OutMap[FLinuxPlatformProperties<false, false, false, false>::PlatformName()]);
+			ITargetPlatformControls* EditorTP = new TLinuxTargetPlatformControls<FLinuxPlatformProperties<true, false, false, false>>(OutMap[FLinuxPlatformProperties<true, false, false, false>::PlatformName()]);
+			ITargetPlatformControls* ServerTP = new TLinuxTargetPlatformControls<FLinuxPlatformProperties<false, true, false, false>>(OutMap[FLinuxPlatformProperties<false, true, false, false>::PlatformName()]);
+			ITargetPlatformControls* ClientTP = new TLinuxTargetPlatformControls<FLinuxPlatformProperties<false, false, true, false>>(OutMap[FLinuxPlatformProperties<false, false, true, false>::PlatformName()]);
+			TargetPlatforms.Add(GameTP);
+			TargetPlatforms.Add(EditorTP);
+			TargetPlatforms.Add(ServerTP);
+			TargetPlatforms.Add(ClientTP);
+		}
+	}
+
+public:
+
+	// IModuleInterface interface
+
+	virtual void StartupModule() override
+	{
+	}
+
+	virtual void ShutdownModule() override
+	{
+
+	}
+};
+
+
+#undef LOCTEXT_NAMESPACE
+
+
+IMPLEMENT_MODULE( FLinuxTargetPlatformControlsModule, LinuxTargetPlatformControls);
diff --git a/Engine/Source/Developer/Linux/LinuxTargetPlatformControls/Public/LinuxTargetDevice.h b/Engine/Source/Developer/Linux/LinuxTargetPlatformControls/Public/LinuxTargetDevice.h
new file mode 100644
--- /dev/null
+++ b/Engine/Source/Developer/Linux/LinuxTargetPlatformControls/Public/LinuxTargetDevice.h
@@ -0,0 +1,223 @@
+// Copyright Epic Games, Inc. All Rights Reserved.
+
+#pragma once
+
+#include "CoreMinimal.h"
+#include "Misc/Paths.h"
+#include "HAL/FileManager.h"
+#include "HAL/PlatformProcess.h"
+#include "Interfaces/ITargetPlatform.h"
+#include "Interfaces/ITargetDevice.h"
+#include "Interfaces/TargetDeviceId.h"
+
+#if PLATFORM_LINUX
+	#include <signal.h>
+	#include <pwd.h>
+#endif // PLATFORM_LINUX
+
+class IFileManager;
+struct FProcHandle;
+
+
+/**
+ * Type definition for shared pointers to instances of FLinuxTargetDevice.
+ */
+typedef TSharedPtr<class FLinuxTargetDevice, ESPMode::ThreadSafe> FLinuxTargetDevicePtr;
+
+/**
+ * Type definition for shared references to instances of FLinuxTargetDevice.
+ */
+typedef TSharedRef<class FLinuxTargetDevice, ESPMode::ThreadSafe> FLinuxTargetDeviceRef;
+
+
+/**
+ * Implements a Linux target device.
+ */
+class FLinuxTargetDevice
+	: public ITargetDevice
+{
+public:
+
+	/**
+	 * Creates and initializes a new device for the specified target platform.
+	 *
+	 * @param InTargetPlatformControls - The target platform controls.
+	 */
+	FLinuxTargetDevice(const ITargetPlatformControls& InTargetPlatformControls, const FString& InDeviceName, TFunction<void()> InSavePlatformDevices)
+		: TargetPlatformControls(InTargetPlatformControls)
+		, DeviceName(InDeviceName)
+		, SavePlatformDevices(InSavePlatformDevices)
+	{ }
+
+	/**
+	 * Minimal constructor for use with SteamDeck devices */
+	FLinuxTargetDevice(const ITargetPlatformControls& InTargetPlatformControls)
+		: FLinuxTargetDevice(InTargetPlatformControls, TEXT("UnknownName"), nullptr)
+	{ }
+
+public:
+
+	virtual bool Connect( ) override
+	{
+		return true;
+	}
+
+	virtual void Disconnect( ) override
+	{ }
+
+	virtual ETargetDeviceTypes GetDeviceType( ) const override
+	{
+		return ETargetDeviceTypes::Desktop;
+	}
+
+	virtual FTargetDeviceId GetId() const override
+	{
+		return FTargetDeviceId(TargetPlatformControls.PlatformName(), GetName());
+	}
+
+	virtual FString GetName( ) const override
+	{
+		return DeviceName;
+	}
+
+	virtual FString GetOperatingSystemName( ) override
+	{
+		return TEXT("GNU/Linux");
+	}
+
+	virtual int32 GetProcessSnapshot( TArray<FTargetDeviceProcessInfo>& OutProcessInfos ) override
+	{
+		STUBBED("FLinuxTargetDevice::GetProcessSnapshot");
+		return 0;
+	}
+
+	virtual const class ITargetPlatformSettings& GetPlatformSettings() const override
+	{
+		return *(TargetPlatformControls.GetTargetPlatformSettings());
+	}
+	virtual const class ITargetPlatformControls& GetPlatformControls() const override
+	{
+		return TargetPlatformControls;
+	}
+
+	virtual bool IsConnected( ) override
+	{
+		return true;
+	}
+
+	virtual bool IsDefault( ) const override
+	{
+		return true;
+	}
+
+	virtual bool PowerOff( bool Force ) override
+	{
+		return false;
+	}
+
+	virtual bool PowerOn( ) override
+	{
+		return false;
+	}
+
+	virtual bool Reboot( bool bReconnect = false ) override
+	{
+		STUBBED("FLinuxTargetDevice::Reboot");
+		return false;
+	}
+
+	virtual bool SupportsFeature( ETargetDeviceFeatures Feature ) const override
+	{
+		switch (Feature)
+		{
+		case ETargetDeviceFeatures::MultiLaunch:
+			return true;
+
+			// @todo to be implemented
+		case ETargetDeviceFeatures::PowerOff:
+			return false;
+
+			// @todo to be implemented turning on remote PCs (wake on LAN)
+		case ETargetDeviceFeatures::PowerOn:
+			return false;
+
+			// @todo to be implemented
+		case ETargetDeviceFeatures::ProcessSnapshot:
+			return false;
+
+			// @todo to be implemented
+		case ETargetDeviceFeatures::Reboot:
+			return false;
+		}
+
+		return false;
+	}
+
+	virtual void SetUserCredentials( const FString& InUserName, const FString& InUserPassword ) override
+	{
+		UserName = InUserName;
+		UserPassword = InUserPassword;
+
+		if (SavePlatformDevices)
+		{
+			SavePlatformDevices();
+		}
+	}
+
+	virtual bool GetUserCredentials( FString& OutUserName, FString& OutUserPassword ) override
+	{
+		OutUserName = UserName;
+		OutUserPassword = UserPassword;
+		return true;
+	}
+
+	virtual bool TerminateProcess( const int64 ProcessId ) override
+	{
+#if PLATFORM_LINUX // if running natively, just terminate the local process
+		// get process path from the ProcessId
+		const int32 ReadLinkSize = 1024;
+		char ReadLinkCmd[ReadLinkSize] = { 0 };
+		FCStringAnsi::Sprintf(ReadLinkCmd, "/proc/%lld/exe", ProcessId);
+		char ProcessPath[UNIX_MAX_PATH + 1] = { 0 };
+		int32 Ret = readlink(ReadLinkCmd, ProcessPath, UE_ARRAY_COUNT(ProcessPath) - 1);
+		if (Ret != -1)
+		{
+			struct stat st;
+			uid_t euid;
+			stat(ProcessPath, &st);
+			euid = geteuid(); // get effective uid of current application, as this user is asking to kill a process
+
+							  // now check if we own the process
+			if (st.st_uid == euid)
+			{
+				// terminate it (will this terminate children as well because we set their pgid?)
+				kill(ProcessId, SIGTERM);
+				sleep(2); // sleep in case process still remains then send a more strict signal
+				kill(ProcessId, SIGKILL);
+				return true;
+			}
+		}
+#else
+		// @todo: support remote termination
+		STUBBED("FLinuxTargetDevice::TerminateProcess");
+#endif // PLATFORM_LINUX
+		return false;
+	}
+
+protected:
+
+	// Holds a reference to the device's target platform.
+	const ITargetPlatformControls& TargetPlatformControls;
+
+	/** Device display name */
+	FString DeviceName;
+
+	/** User name on the remote machine */
+	FString UserName;
+
+	/** User password on the remote machine */
+	FString UserPassword;
+
+	/** Target platform function to save device state */
+	TFunction<void()> SavePlatformDevices;
+};
diff --git a/Engine/Source/Developer/Linux/LinuxTargetPlatformControls/Public/LinuxTargetPlatformControls.h b/Engine/Source/Developer/Linux/LinuxTargetPlatformControls/Public/LinuxTargetPlatformControls.h
new file mode 100644
--- /dev/null
+++ b/Engine/Source/Developer/Linux/LinuxTargetPlatformControls/Public/LinuxTargetPlatformControls.h
@@ -0,0 +1,378 @@
+// Copyright Epic Games, Inc. All Rights Reserved.
+
+/*=============================================================================
+	LinuxTargetPlatformControls.h: Declares the TLinuxTargetPlatformControls class.
+=============================================================================*/
+
+#pragma once
+
+#include "CoreMinimal.h"
+#include "Interfaces/TargetDeviceId.h"
+#include "Common/TargetPlatformControlsBase.h"
+#include "SteamDeck/SteamDeckDevice.h"
+#include "AnalyticsEventAttribute.h"
+#include "InstalledPlatformInfo.h"
+#include "LinuxTargetDevice.h"
+#include "Linux/LinuxPlatformProperties.h"
+#if WITH_ENGINE
+#include "TextureResource.h"
+#endif // WITH_ENGINE
+#define LOCTEXT_NAMESPACE "TLinuxTargetPlatformControls"
+
+
+/**
+ * Template for Linux target platforms controls
+ */
+template<typename TProperties>
+class TLinuxTargetPlatformControls
+	: public TTargetPlatformControlsBase<TProperties>
+{
+public:
+
+	typedef TTargetPlatformControlsBase<TProperties> TSuper;
+
+	/**
+	 * Default constructor.
+	 */
+	TLinuxTargetPlatformControls(ITargetPlatformSettings* TargetPlatformSettings)
+		: TSuper(TargetPlatformSettings)
+#if WITH_ENGINE
+		, bChangingDeviceConfig(false)
+#endif // WITH_ENGINE
+	{
+#if PLATFORM_LINUX
+		if (!TProperties::IsArm64())
+		{
+			// only add local device if actually running on Linux
+			LocalDevice = MakeShareable(new FLinuxTargetDevice(*this, FPlatformProcess::ComputerName(), nullptr));
+		}
+#endif
+
+#if WITH_ENGINE
+
+		InitDevicesFromConfig();
+
+		if (!TProperties::IsArm64())
+		{
+			SteamDevices = TSteamDeckDevice<FLinuxTargetDevice>::DiscoverDevices(*this, TEXT("Native Linux"));
+		}
+
+#endif // WITH_ENGINE
+	}
+
+
+public:
+
+	//~ Begin ITargetPlatform Interface
+
+	virtual void EnableDeviceCheck(bool OnOff) override {}
+
+	virtual bool AddDevice(const FString& DeviceName, bool bDefault) override
+	{
+		return AddDevice(DeviceName, TEXT(""), TEXT(""), TEXT(""), bDefault);
+	}
+
+	virtual bool AddDevice(const FString& DeviceName, const FString& DeviceUserFriendlyName, const FString& Username, const FString& Password, bool bDefault) override
+	{
+		FLinuxTargetDevicePtr& Device = Devices.FindOrAdd(DeviceName);
+
+		if (Device.IsValid())
+		{
+			// do not allow duplicates
+			return false;
+		}
+
+		Device = MakeShareable(new FLinuxTargetDevice(*this, DeviceName,
+#if WITH_ENGINE
+			[&]() { SaveDevicesToConfig(); }));
+		SaveDevicesToConfig();	// this will do the right thing even if AddDevice() was called from InitDevicesFromConfig
+#else
+			nullptr));
+#endif // WITH_ENGINE
+
+		if (!Username.IsEmpty() || !Password.IsEmpty())
+		{
+			Device->SetUserCredentials(Username, Password);
+		}
+
+		ITargetPlatformControls::OnDeviceDiscovered().Broadcast(Device.ToSharedRef());
+		return true;
+	}
+
+
+	virtual void GetAllDevices(TArray<ITargetDevicePtr>& OutDevices) const override
+	{
+		// TODO: ping all the machines in a local segment and/or try to connect to port 22 of those that respond
+		OutDevices.Reset();
+		if (LocalDevice.IsValid())
+		{
+			OutDevices.Add(LocalDevice);
+		}
+
+		for (const auto& DeviceIter : Devices)
+		{
+			OutDevices.Add(DeviceIter.Value);
+		}
+
+		for (const ITargetDevicePtr& SteamDeck : SteamDevices)
+		{
+			if (SteamDeck.IsValid())
+			{
+				OutDevices.Add(SteamDeck);
+			}
+		}
+	}
+
+	virtual bool GenerateStreamingInstallManifest(const TMultiMap<FString, int32>& PakchunkMap, const TSet<int32>& PakchunkIndicesInUse) const override
+	{
+		return true;
+	}
+
+	virtual ITargetDevicePtr GetDefaultDevice() const override
+	{
+		if (LocalDevice.IsValid())
+		{
+			return LocalDevice;
+		}
+
+		return nullptr;
+	}
+
+	virtual ITargetDevicePtr GetDevice(const FTargetDeviceId& DeviceId) override
+	{
+		if (LocalDevice.IsValid() && (DeviceId == LocalDevice->GetId()))
+		{
+			return LocalDevice;
+		}
+
+		for (const auto& DeviceIter : Devices)
+		{
+			if (DeviceId == DeviceIter.Value->GetId())
+			{
+				return DeviceIter.Value;
+			}
+		}
+
+		for (const ITargetDevicePtr& SteamDeck : SteamDevices)
+		{
+			if (SteamDeck.IsValid() && DeviceId == SteamDeck->GetId())
+			{
+				return SteamDeck;
+			}
+		}
+
+		return nullptr;
+	}
+
+	virtual bool IsRunningPlatform() const override
+	{
+		// Must be Linux platform as editor for this to be considered a running platform
+		return PLATFORM_LINUX && !UE_SERVER && !UE_GAME && WITH_EDITOR && TProperties::HasEditorOnlyData();
+	}
+
+	virtual bool IsSdkInstalled(bool bProjectHasCode, FString& OutDocumentationPath) const override
+	{
+		if (!PLATFORM_LINUX)
+		{
+			// check for LINUX_MULTIARCH_ROOT or for legacy LINUX_ROOT when targeting Linux from Win/Mac
+
+			// proceed with any value for MULTIARCH root, because checking exact architecture is not possible at this point
+			FString ToolchainMultiarchRoot = FPlatformMisc::GetEnvironmentVariable(TEXT("LINUX_MULTIARCH_ROOT"));
+			if (ToolchainMultiarchRoot.Len() > 0 && FPaths::DirectoryExists(ToolchainMultiarchRoot))
+			{
+				return true;
+			}
+
+			// else check for legacy LINUX_ROOT
+			FString ToolchainCompiler = FPlatformMisc::GetEnvironmentVariable(TEXT("LINUX_ROOT"));
+			if (PLATFORM_WINDOWS)
+			{
+				ToolchainCompiler += "/bin/clang++.exe";
+			}
+			else if (PLATFORM_MAC)
+			{
+				ToolchainCompiler += "/bin/clang++";
+			}
+			else
+			{
+				checkf(false, TEXT("Unable to target Linux on an unknown platform."));
+				return false;
+			}
+
+			return FPaths::FileExists(ToolchainCompiler);
+		}
+
+		return true;
+	}
+
+	virtual int32 CheckRequirements(bool bProjectHasCode, EBuildConfiguration Configuration, bool bRequiresAssetNativization, FString& OutTutorialPath, FString& OutDocumentationPath, FText& CustomizedLogMessage) const override
+	{
+		int32 ReadyToBuild = TSuper::CheckRequirements(bProjectHasCode, Configuration, bRequiresAssetNativization, OutTutorialPath, OutDocumentationPath, CustomizedLogMessage);
+
+		// do not support code/plugins in Installed builds if the required libs aren't bundled (on Windows/Mac)
+		if (!PLATFORM_LINUX && !FInstalledPlatformInfo::Get().IsValidPlatform(TSuper::GetPlatformInfo().UBTPlatformString, EProjectType::Code))
+		{
+			if (bProjectHasCode)
+			{
+				ReadyToBuild |= ETargetPlatformReadyStatus::CodeUnsupported;
+			}
+
+			FText Reason;
+			if (this->RequiresTempTarget(bProjectHasCode, Configuration, bRequiresAssetNativization, Reason))
+			{
+				ReadyToBuild |= ETargetPlatformReadyStatus::PluginsUnsupported;
+			}
+		}
+
+		return ReadyToBuild;
+	}
+
+	virtual void GetPlatformSpecificProjectAnalytics(TArray<FAnalyticsEventAttribute>& AnalyticsParamArray) const override
+	{
+		TSuper::GetPlatformSpecificProjectAnalytics(AnalyticsParamArray);
+
+		TSuper::AppendAnalyticsEventConfigArray(AnalyticsParamArray, TEXT("/Script/LinuxTargetPlatform.LinuxTargetSettings"), TEXT("TargetedRHIs"), GEngineIni);
+	}
+
+#if WITH_ENGINE
+
+	virtual void GetTextureFormats(const UTexture* InTexture, TArray< TArray<FName> >& OutFormats) const override
+	{
+		if (this->AllowAudioVisualData())
+		{
+			// just use the standard texture format name for this texture
+			GetDefaultTextureFormatNamePerLayer(OutFormats.AddDefaulted_GetRef(), this->GetTargetPlatformSettings(), this, InTexture, true, 4, true);
+		}
+	}
+
+	virtual void GetAllTextureFormats(TArray<FName>& OutFormats) const override
+	{
+		if (this->AllowAudioVisualData())
+		{
+			// just use the standard texture format name for this texture
+			GetAllDefaultTextureFormats(this->GetTargetPlatformSettings(), OutFormats);
+		}
+	}
+
+#endif //WITH_ENGINE
+
+	virtual bool SupportsVariants() const override
+	{
+		return true;
+	}
+
+	virtual float GetVariantPriority() const override
+	{
+		return TProperties::GetVariantPriority();
+	}
+
+	//~ End ITargetPlatform Interface
+
+protected:
+
+#if WITH_ENGINE
+	/** Whether we're in process of changing device config - if yes, we will prevent recurrent calls. */
+	bool bChangingDeviceConfig;
+
+	void InitDevicesFromConfig()
+	{
+		if (bChangingDeviceConfig)
+		{
+			return;
+		}
+		bChangingDeviceConfig = true;
+
+		int NumDevices = 0;
+		for (;; ++NumDevices)
+		{
+			FString DeviceName, DeviceUser, DevicePass;
+
+			FString DeviceBaseKey(FString::Printf(TEXT("LinuxTargetPlatfrom_%s_Device_%d"), *TSuper::PlatformName(), NumDevices));
+			FString DeviceNameKey = DeviceBaseKey + TEXT("_Name");
+			if (!GConfig->GetString(TEXT("/Script/LinuxTargetPlatform.LinuxTargetSettings"), *DeviceNameKey, DeviceName, GEngineIni))
+			{
+				// no such device
+				break;
+			}
+
+			if (!AddDevice(DeviceName, false))
+			{
+				break;
+			}
+
+			// set credentials, if any
+			FString DeviceUserKey = DeviceBaseKey + TEXT("_User");
+			if (GConfig->GetString(TEXT("/Script/LinuxTargetPlatform.LinuxTargetSettings"), *DeviceUserKey, DeviceUser, GEngineIni))
+			{
+				FString DevicePassKey = DeviceBaseKey + TEXT("_Pass");
+				if (GConfig->GetString(TEXT("/Script/LinuxTargetPlatform.LinuxTargetSettings"), *DevicePassKey, DevicePass, GEngineIni))
+				{
+					for (const auto& DeviceIter : Devices)
+					{
+						ITargetDevicePtr Device = DeviceIter.Value;
+						if (Device.IsValid() && Device->GetId().GetDeviceName() == DeviceName)
+						{
+							Device->SetUserCredentials(DeviceUser, DevicePass);
+						}
+					}
+				}
+			}
+		}
+
+		bChangingDeviceConfig = false;
+	}
+
+	void SaveDevicesToConfig()
+	{
+		if (bChangingDeviceConfig)
+		{
+			return;
+		}
+		bChangingDeviceConfig = true;
+
+		int DeviceIndex = 0;
+		for (const auto& DeviceIter : Devices)
+		{
+			ITargetDevicePtr Device = DeviceIter.Value;
+
+			FString DeviceBaseKey(FString::Printf(TEXT("LinuxTargetPlatfrom_%s_Device_%d"), *TSuper::PlatformName(), DeviceIndex));
+			FString DeviceNameKey = DeviceBaseKey + TEXT("_Name");
+
+			if (Device.IsValid())
+			{
+				FString DeviceName = Device->GetId().GetDeviceName();
+				// do not save a local device on Linux or it will be duplicated
+				if (PLATFORM_LINUX && DeviceName == FPlatformProcess::ComputerName())
+				{
+					continue;
+				}
+
+				GConfig->SetString(TEXT("/Script/LinuxTargetPlatform.LinuxTargetSettings"), *DeviceNameKey, *DeviceName, GEngineIni);
+
+				FString DeviceUser, DevicePass;
+				if (Device->GetUserCredentials(DeviceUser, DevicePass))
+				{
+					FString DeviceUserKey = DeviceBaseKey + TEXT("_User");
+					FString DevicePassKey = DeviceBaseKey + TEXT("_Pass");
+
+					GConfig->SetString(TEXT("/Script/LinuxTargetPlatform.LinuxTargetSettings"), *DeviceUserKey, *DeviceUser, GEngineIni);
+					GConfig->SetString(TEXT("/Script/LinuxTargetPlatform.LinuxTargetSettings"), *DevicePassKey, *DevicePass, GEngineIni);
+				}
+
+				++DeviceIndex;	// needs to be incremented here since we cannot allow gaps
+			}
+		}
+
+		bChangingDeviceConfig = false;
+	}
+#endif // WITH_ENGINE
+
+	// Holds the local device.
+	FLinuxTargetDevicePtr LocalDevice;
+	// Holds a map of valid devices.
+	TMap<FString, FLinuxTargetDevicePtr> Devices;
+	TArray<ITargetDevicePtr> SteamDevices;
+
+};
+
+#undef LOCTEXT_NAMESPACE
diff --git a/Engine/Source/Developer/Linux/LinuxTargetPlatformSettings/Classes/LinuxTargetSettings.h b/Engine/Source/Developer/Linux/LinuxTargetPlatformSettings/Classes/LinuxTargetSettings.h
new file mode 100644
--- /dev/null
+++ b/Engine/Source/Developer/Linux/LinuxTargetPlatformSettings/Classes/LinuxTargetSettings.h
@@ -0,0 +1,52 @@
+// Copyright Epic Games, Inc. All Rights Reserved.
+
+/*=============================================================================
+	LinuxTargetSettings.h: Declares the ULinuxTargetSettings class.
+=============================================================================*/
+
+#pragma once
+
+#include "CoreMinimal.h"
+#include "UObject/ObjectMacros.h"
+#include "UObject/Object.h"
+#include "LinuxTargetSettings.generated.h"
+
+
+/**
+ * Implements the settings for the Linux target platform.
+ */
+UCLASS(config=Engine, defaultconfig)
+class LINUXTARGETPLATFORMSETTINGS_API ULinuxTargetSettings
+	: public UObject
+{
+public:
+
+	GENERATED_UCLASS_BODY()
+
+	/** Which of the currently enabled spatialization plugins to use. */
+	UPROPERTY(config, EditAnywhere, Category = "Audio")
+	FString SpatializationPlugin;
+
+	/** Which of the currently enabled source data override plugins to use. */
+	UPROPERTY(config, EditAnywhere, Category = "Audio")
+	FString SourceDataOverridePlugin;
+
+	/** Which of the currently enabled reverb plugins to use. */
+	UPROPERTY(config, EditAnywhere, Category = "Audio")
+	FString ReverbPlugin;
+
+	/** Which of the currently enabled occlusion plugins to use. */
+	UPROPERTY(config, EditAnywhere, Category = "Audio")
+	FString OcclusionPlugin;
+
+	/** Quality Level to COOK SoundCues at (if set, all other levels will be stripped by the cooker). */
+	UPROPERTY(config, EditAnywhere, Category = "Audio|CookOverrides", meta = (DisplayName = "Sound Cue Cook Quality"))
+	int32 SoundCueCookQualityIndex = INDEX_NONE;
+
+	/**
+	 * The collection of RHI's we want to support on this platform.
+	 * This is not always the full list of RHI we can support.
+	 */
+	UPROPERTY(EditAnywhere, config, Category=Rendering)
+	TArray<FString> TargetedRHIs;
+};
diff --git a/Engine/Source/Developer/Linux/LinuxTargetPlatformSettings/LinuxTargetPlatformSettings.Build.cs b/Engine/Source/Developer/Linux/LinuxTargetPlatformSettings/LinuxTargetPlatformSettings.Build.cs
new file mode 100644
--- /dev/null
+++ b/Engine/Source/Developer/Linux/LinuxTargetPlatformSettings/LinuxTargetPlatformSettings.Build.cs
@@ -0,0 +1,37 @@
+// Copyright Epic Games, Inc. All Rights Reserved.
+
+using UnrealBuildTool;
+
+public class LinuxTargetPlatformSettings : ModuleRules
+{
+    public LinuxTargetPlatformSettings(ReadOnlyTargetRules Target) : base(Target)
+	{
+        BinariesSubFolder = "Linux";
+		// We need a short name here since this can run afoul very easily of the `MAX_PATH` limit when
+		// combined with building other targets that make use of this.
+		ShortName = "LinuxTPSet";
+
+		PrivateDependencyModuleNames.AddRange(
+            new string[] {
+				"Core",
+				"CoreUObject",
+				"TargetPlatform",
+				"DesktopPlatform",
+				"Projects"
+			}
+        );
+
+		PrivateIncludePathModuleNames.AddRange(
+			new string[] {
+				"Settings"
+			}
+		);
+
+		if (Target.bCompileAgainstEngine)
+		{
+			PublicIncludePathModuleNames.Add("Engine");
+			PrivateDependencyModuleNames.Add("Engine");
+			PrivateIncludePathModuleNames.Add("TextureCompressor");
+		}
+	}
+}
diff --git a/Engine/Source/Developer/Linux/LinuxTargetPlatformSettings/Private/LinuxTargetPlatformClasses.cpp b/Engine/Source/Developer/Linux/LinuxTargetPlatformSettings/Private/LinuxTargetPlatformClasses.cpp
new file mode 100644
--- /dev/null
+++ b/Engine/Source/Developer/Linux/LinuxTargetPlatformSettings/Private/LinuxTargetPlatformClasses.cpp
@@ -0,0 +1,17 @@
+// Copyright Epic Games, Inc. All Rights Reserved.
+
+/*=============================================================================
+	LinuxTargetPlatformClasses.cpp: Implements the module's UClasses.
+=============================================================================*/
+
+#include "CoreMinimal.h"
+#include "LinuxTargetSettings.h"
+#include "LinuxTargetPlatformSettings.h"
+
+
+/* ULinuxTargetSettings structors
+ *****************************************************************************/
+
+ULinuxTargetSettings::ULinuxTargetSettings( const FObjectInitializer& ObjectInitializer )
+	: Super(ObjectInitializer)
+{ }
diff --git a/Engine/Source/Developer/Linux/LinuxTargetPlatformSettings/Private/LinuxTargetPlatformSettingsModule.cpp b/Engine/Source/Developer/Linux/LinuxTargetPlatformSettings/Private/LinuxTargetPlatformSettingsModule.cpp
new file mode 100644
--- /dev/null
+++ b/Engine/Source/Developer/Linux/LinuxTargetPlatformSettings/Private/LinuxTargetPlatformSettingsModule.cpp
@@ -0,0 +1,104 @@
+// Copyright Epic Games, Inc. All Rights Reserved.
+
+#include "CoreMinimal.h"
+#include "Modules/ModuleManager.h"
+#include "UObject/Package.h"
+#include "UObject/WeakObjectPtr.h"
+#include "ILinuxTargetPlatformSettingsModule.h"
+#include "LinuxTargetSettings.h"
+#include "LinuxTargetPlatformSettings.h"
+#include "ISettingsModule.h"
+
+
+#define LOCTEXT_NAMESPACE "FLinuxTargetPlatformSettingsModule"
+
+
+
+/**
+ * Module for the Linux target platform settings.
+ */
+class FLinuxTargetPlatformSettingsModule
+	: public ILinuxTargetPlatformSettingsModule
+{
+public:
+	virtual void GetTargetPlatformSettings(TArray<ITargetPlatformSettings*>& TargetPlatforms) override
+	{
+		// NoEditor TP
+		ITargetPlatformSettings* NoEditorTP = new TLinuxTargetPlatformSettings<FLinuxPlatformProperties<false, false, false, false>>();
+		TargetPlatforms.Add(NoEditorTP);
+		PlatformNameToPlatformSettings.Add(FLinuxPlatformProperties<false, false, false, false>::PlatformName(), NoEditorTP);
+		// Editor TP
+		ITargetPlatformSettings* EditorTP = new TLinuxTargetPlatformSettings<FLinuxPlatformProperties<true, false, false, false>>();
+		TargetPlatforms.Add(EditorTP);
+		PlatformNameToPlatformSettings.Add(FLinuxPlatformProperties<true, false, false, false>::PlatformName(), EditorTP);
+		// Server TP
+		ITargetPlatformSettings* ServerTP = new TLinuxTargetPlatformSettings<FLinuxPlatformProperties<false, true, false, false>>();
+		TargetPlatforms.Add(ServerTP);
+		PlatformNameToPlatformSettings.Add(FLinuxPlatformProperties<false, true, false, false>::PlatformName(), ServerTP);
+		// Client TP
+		ITargetPlatformSettings* ClientTP = new TLinuxTargetPlatformSettings<FLinuxPlatformProperties<false, false, true, false>>();
+		TargetPlatforms.Add(ClientTP);
+		PlatformNameToPlatformSettings.Add(FLinuxPlatformProperties<false, false, true, false>::PlatformName(), ClientTP);
+	}
+
+	virtual void GetPlatformSettingsMaps(TMap<FString, ITargetPlatformSettings*>& OutMap) override
+	{
+		OutMap = PlatformNameToPlatformSettings;
+	}
+
+public:
+
+	// IModuleInterface interface
+
+	virtual void StartupModule() override
+	{
+		TargetSettings = NewObject<ULinuxTargetSettings>(GetTransientPackage(), "LinuxTargetSettings", RF_Standalone);
+
+		// We need to manually load the config properties here, as this module is loaded before the UObject system is setup to do this
+		GConfig->GetArray(TEXT("/Script/LinuxTargetPlatform.LinuxTargetSettings"), TEXT("TargetedRHIs"), TargetSettings->TargetedRHIs, GEngineIni);
+		TargetSettings->AddToRoot();
+
+		ISettingsModule* SettingsModule = FModuleManager::GetModulePtr<ISettingsModule>("Settings");
+
+		if (SettingsModule != nullptr)
+		{
+			SettingsModule->RegisterSettings("Project", "Platforms", "Linux",
+				LOCTEXT("TargetSettingsName", "Linux"),
+				LOCTEXT("TargetSettingsDescription", "Settings for Linux target platform"),
+				TargetSettings
+			);
+		}
+	}
+
+	virtual void ShutdownModule() override
+	{
+		ISettingsModule* SettingsModule = FModuleManager::GetModulePtr<ISettingsModule>("Settings");
+
+		if (SettingsModule != nullptr)
+		{
+			SettingsModule->UnregisterSettings("Project", "Platforms", "Linux");
+		}
+
+		if (!GExitPurge)
+		{
+			// If we're in exit purge, this object has already been destroyed
+			TargetSettings->RemoveFromRoot();
+		}
+		else
+		{
+			TargetSettings = NULL;
+		}
+	}
+
+private:
+	TMap<FString, ITargetPlatformSettings*> PlatformNameToPlatformSettings;
+
+	/** Holds the target settings. */
+	ULinuxTargetSettings* TargetSettings;
+};
+
+
+#undef LOCTEXT_NAMESPACE
+
+
+IMPLEMENT_MODULE( FLinuxTargetPlatformSettingsModule, LinuxTargetPlatformSettings);
diff --git a/Engine/Source/Developer/Linux/LinuxTargetPlatformSettings/Public/ILinuxTargetPlatformSettingsModule.h b/Engine/Source/Developer/Linux/LinuxTargetPlatformSettings/Public/ILinuxTargetPlatformSettingsModule.h
new file mode 100644
--- /dev/null
+++ b/Engine/Source/Developer/Linux/LinuxTargetPlatformSettings/Public/ILinuxTargetPlatformSettingsModule.h
@@ -0,0 +1,16 @@
+// Copyright Epic Games, Inc. All Rights Reserved.
+
+#pragma  once 
+
+/*------------------------------------------------------------------------------------
+ILinuxTargetPlatformSettingsModule interface
+------------------------------------------------------------------------------------*/
+
+#include "CoreMinimal.h"
+#include "Interfaces/ITargetPlatformSettingsModule.h"
+
+class ILinuxTargetPlatformSettingsModule : public ITargetPlatformSettingsModule
+{
+public:
+	virtual void GetPlatformSettingsMaps(TMap<FString, ITargetPlatformSettings*>& OutMap) = 0;
+};
diff --git a/Engine/Source/Developer/Linux/LinuxTargetPlatformSettings/Public/LinuxTargetPlatformSettings.h b/Engine/Source/Developer/Linux/LinuxTargetPlatformSettings/Public/LinuxTargetPlatformSettings.h
new file mode 100644
--- /dev/null
+++ b/Engine/Source/Developer/Linux/LinuxTargetPlatformSettings/Public/LinuxTargetPlatformSettings.h
@@ -0,0 +1,157 @@
+// Copyright Epic Games, Inc. All Rights Reserved.
+
+/*=============================================================================
+	LinuxTargetPlatform.h: Declares the TLinuxTargetPlatformSettings class.
+=============================================================================*/
+
+#pragma once
+
+#include "CoreMinimal.h"
+#include "Common/TargetPlatformSettingsBase.h"
+#if WITH_ENGINE
+#include "StaticMeshResources.h"
+#endif // WITH_ENGINE
+#include "Linux/LinuxPlatformProperties.h"
+
+#define LOCTEXT_NAMESPACE "TLinuxTargetPlatformSettings"
+
+class UTextureLODSettings;
+
+/**
+ * Template for Linux target platforms settings
+ */
+template<typename TProperties>
+class TLinuxTargetPlatformSettings
+	: public TTargetPlatformSettingsBase<TProperties>
+{
+public:
+
+	typedef TTargetPlatformSettingsBase<TProperties> TSuper;
+
+	/**
+	 * Default constructor.
+	 */
+	TLinuxTargetPlatformSettings( )
+
+	{
+#if WITH_ENGINE
+		TextureLODSettings = nullptr;
+		StaticMeshLODSettings.Initialize(this);
+
+		// Get the Target RHIs for this platform, we do not always want all those that are supported.
+		TArray<FName> TargetedShaderFormats;
+		TLinuxTargetPlatformSettings::GetAllTargetedShaderFormats(TargetedShaderFormats);
+
+		// If we are targeting ES 2.0/3.1, we also must cook encoded HDR reflection captures
+		static FName NAME_SF_VULKAN_ES31(TEXT("SF_VULKAN_ES31"));
+		static FName NAME_OPENGL_150_ES3_1(TEXT("GLSL_150_ES31"));
+		bRequiresEncodedHDRReflectionCaptures = TargetedShaderFormats.Contains(NAME_SF_VULKAN_ES31)
+			|| TargetedShaderFormats.Contains(NAME_OPENGL_150_ES3_1);
+#endif // WITH_ENGINE
+	}
+
+
+public:
+
+	//~ Begin ITargetPlatform Interface
+
+	virtual bool SupportsFeature(ETargetPlatformFeatures Feature) const override
+	{
+		if (Feature == ETargetPlatformFeatures::UserCredentials || Feature == ETargetPlatformFeatures::Packaging)
+		{
+			return true;
+		}
+
+		if (Feature == ETargetPlatformFeatures::CanCookPackages)
+		{
+			return TProperties::HasEditorOnlyData();
+		}
+
+		return TSuper::SupportsFeature(Feature);
+	}
+
+	virtual void GetAllPossibleShaderFormats( TArray<FName>& OutFormats ) const override
+	{
+		if (!TProperties::IsServerOnly())
+		{
+			static FName NAME_VULKAN_SM5(TEXT("SF_VULKAN_SM5"));
+			static FName NAME_VULKAN_SM6(TEXT("SF_VULKAN_SM6"));
+			static FName NAME_VULKAN_ES31(TEXT("SF_VULKAN_ES31"));
+
+			OutFormats.AddUnique(NAME_VULKAN_SM5);
+			OutFormats.AddUnique(NAME_VULKAN_SM6);
+			OutFormats.AddUnique(NAME_VULKAN_ES31);
+		}
+	}
+
+	virtual void GetAllTargetedShaderFormats( TArray<FName>& OutFormats ) const override
+	{
+		// Get the Target RHIs for this platform, we do not always want all those that are supported. (reload in case user changed in the editor)
+		TArray<FString>TargetedShaderFormats;
+		GConfig->GetArray(TEXT("/Script/LinuxTargetPlatform.LinuxTargetSettings"), TEXT("TargetedRHIs"), TargetedShaderFormats, GEngineIni);
+
+		// Gather the list of Target RHIs and filter out any that may be invalid.
+		TArray<FName> PossibleShaderFormats;
+		GetAllPossibleShaderFormats(PossibleShaderFormats);
+
+		for (int32 ShaderFormatIdx = TargetedShaderFormats.Num() - 1; ShaderFormatIdx >= 0; ShaderFormatIdx--)
+		{
+			FString ShaderFormat = TargetedShaderFormats[ShaderFormatIdx];
+			if (PossibleShaderFormats.Contains(FName(*ShaderFormat)) == false)
+			{
+				TargetedShaderFormats.RemoveAt(ShaderFormatIdx);
+			}
+		}
+
+		for(const FString& ShaderFormat : TargetedShaderFormats)
+		{
+			OutFormats.AddUnique(FName(*ShaderFormat));
+		}
+	}
+
+#if WITH_ENGINE
+	virtual void GetReflectionCaptureFormats(TArray<FName>& OutFormats) const override
+	{
+		if (bRequiresEncodedHDRReflectionCaptures)
+		{
+			OutFormats.Add(FName(TEXT("EncodedHDR")));
+		}
+
+		OutFormats.Add(FName(TEXT("FullHDR")));
+	}
+
+	virtual const class FStaticMeshLODSettings& GetStaticMeshLODSettings( ) const override
+	{
+		return StaticMeshLODSettings;
+	}
+
+	virtual const UTextureLODSettings& GetTextureLODSettings() const override
+	{
+		return *TextureLODSettings;
+	}
+
+	virtual void RegisterTextureLODSettings(const UTextureLODSettings* InTextureLODSettings) override
+	{
+		TextureLODSettings = InTextureLODSettings;
+	}
+
+#endif //WITH_ENGINE
+
+	//~ End ITargetPlatform Interface
+
+protected:
+
+
+#if WITH_ENGINE
+	// Holds the texture LOD settings.
+	const UTextureLODSettings* TextureLODSettings;
+
+	// Holds static mesh LOD settings.
+	FStaticMeshLODSettings StaticMeshLODSettings;
+
+	// True if the project requires encoded HDR reflection captures
+	bool bRequiresEncodedHDRReflectionCaptures;
+#endif // WITH_ENGINE
+};
+
+#undef LOCTEXT_NAMESPACE
diff --git a/Engine/Source/Developer/Mac/MacPlatformEditor/MacPlatformEditor.Build.cs b/Engine/Source/Developer/Mac/MacPlatformEditor/MacPlatformEditor.Build.cs
--- a/Engine/Source/Developer/Mac/MacPlatformEditor/MacPlatformEditor.Build.cs
+++ b/Engine/Source/Developer/Mac/MacPlatformEditor/MacPlatformEditor.Build.cs
@@ -20,7 +20,7 @@ public class MacPlatformEditor : ModuleRules
 				"PropertyEditor",
 				"SharedSettingsWidgets",
 				"SourceControl",
-				"MacTargetPlatform",
+				"MacTargetPlatformSettings",
 				"TargetPlatform",
 				"MaterialShaderQualitySettings",
 				"RenderCore",
diff --git a/Engine/Source/Developer/Mac/MacPlatformEditor/Private/MacPlatformEditorModule.cpp b/Engine/Source/Developer/Mac/MacPlatformEditor/Private/MacPlatformEditorModule.cpp
--- a/Engine/Source/Developer/Mac/MacPlatformEditor/Private/MacPlatformEditorModule.cpp
+++ b/Engine/Source/Developer/Mac/MacPlatformEditor/Private/MacPlatformEditorModule.cpp
@@ -4,8 +4,6 @@
 #include "Modules/ModuleInterface.h"
 #include "ISettingsModule.h"
 #include "Modules/ModuleManager.h"
-#include "MacTargetSettings.h"
-#include "XcodeProjectSettings.h"
 #include "UObject/WeakObjectPtr.h"
 #include "UObject/Class.h"
 #include "PropertyEditorModule.h"
@@ -23,39 +21,12 @@ class FMacPlatformEditorModule
 
 	virtual void StartupModule() override
 	{
-		// register settings
-		ISettingsModule* SettingsModule = FModuleManager::GetModulePtr<ISettingsModule>("Settings");
 
-		if (SettingsModule != nullptr)
-		{
-			SettingsModule->RegisterSettings("Project", "Platforms", "Mac",
-				LOCTEXT("MacTargetSettingsName", "Mac"),
-				LOCTEXT("MacTargetSettingsDescription", "Settings for Mac target platform"),
-				GetMutableDefault<UMacTargetSettings>()
-			);
-            SettingsModule->RegisterSettings("Project", "Platforms", "Xcode",
-                LOCTEXT("XcodeProjectSettingsName", "Xcode Projects"),
-                LOCTEXT("XcodeProjectSettingsDescription", "Settings for Xcode projects"),
-                GetMutableDefault<UXcodeProjectSettings>()
-            );
-            
-            FPropertyEditorModule& PropertyModule = FModuleManager::LoadModuleChecked<FPropertyEditorModule>("PropertyEditor");
-            PropertyModule.RegisterCustomClassLayout(
-                "XcodeProjectSettings",
-                FOnGetDetailCustomizationInstance::CreateStatic(&FXcodeProjectSettingsDetailsCustomization::MakeInstance)
-            );
-		}
 	}
 
 	virtual void ShutdownModule() override
 	{
-		ISettingsModule* SettingsModule = FModuleManager::GetModulePtr<ISettingsModule>("Settings");
 
-		if (SettingsModule != nullptr)
-		{
-			SettingsModule->UnregisterSettings("Project", "Platforms", "Mac");
-            SettingsModule->UnregisterSettings("Project", "Platforms", "Xcode");
-		}
 	}
 };
 
diff --git a/Engine/Source/Developer/Mac/MacTargetPlatform/Classes/MacTargetSettings.h b/Engine/Source/Developer/Mac/MacTargetPlatform/Classes/MacTargetSettings.h
deleted file mode 100644
diff --git a/Engine/Source/Developer/Mac/MacTargetPlatform/Classes/XcodeProjectSettings.h b/Engine/Source/Developer/Mac/MacTargetPlatform/Classes/XcodeProjectSettings.h
deleted file mode 100644
diff --git a/Engine/Source/Developer/Mac/MacTargetPlatform/MacTargetPlatform.Build.cs b/Engine/Source/Developer/Mac/MacTargetPlatform/MacTargetPlatform.Build.cs
--- a/Engine/Source/Developer/Mac/MacTargetPlatform/MacTargetPlatform.Build.cs
+++ b/Engine/Source/Developer/Mac/MacTargetPlatform/MacTargetPlatform.Build.cs
@@ -9,22 +9,9 @@ public class MacTargetPlatform : ModuleRules
 		PrivateDependencyModuleNames.AddRange(
 			new string[] {
 				"Core",
-				"CoreUObject",
 				"TargetPlatform",
-				"DesktopPlatform"
+				"DesktopPlatform",
 			}
 		);
-
-		PrivateIncludePathModuleNames.AddRange(
-			new string[] {
-				"Settings"
-			}
-		);
-
-		if (Target.bCompileAgainstEngine)
-		{
-			PrivateDependencyModuleNames.Add("Engine");
-			PrivateIncludePathModuleNames.Add("TextureCompressor");
-		}
 	}
 }
diff --git a/Engine/Source/Developer/Mac/MacTargetPlatform/Private/MacTargetPlatformClasses.cpp b/Engine/Source/Developer/Mac/MacTargetPlatform/Private/MacTargetPlatformClasses.cpp
deleted file mode 100644
diff --git a/Engine/Source/Developer/Mac/MacTargetPlatform/Private/MacTargetPlatformModule.cpp b/Engine/Source/Developer/Mac/MacTargetPlatform/Private/MacTargetPlatformModule.cpp
--- a/Engine/Source/Developer/Mac/MacTargetPlatform/Private/MacTargetPlatformModule.cpp
+++ b/Engine/Source/Developer/Mac/MacTargetPlatform/Private/MacTargetPlatformModule.cpp
@@ -1,15 +1,9 @@
 // Copyright Epic Games, Inc. All Rights Reserved.
 
 #include "CoreMinimal.h"
-#include "GenericMacTargetPlatform.h"
 #include "Modules/ModuleManager.h"
-#include "ISettingsModule.h"
 #include "Interfaces/ITargetPlatformModule.h"
-#include "Modules/ModuleManager.h"
-#include "MacTargetSettings.h"
-#include "XcodeProjectSettings.h"
-#include "UObject/Package.h"
-#include "UObject/WeakObjectPtr.h"
+#include "Common/TargetPlatformBase.h"
 
 #define LOCTEXT_NAMESPACE "FMacTargetPlatformModule"
 
@@ -23,22 +17,15 @@ public:
 
 	virtual void GetTargetPlatforms(TArray<ITargetPlatform*>& TargetPlatforms) override
 	{
-		// Game TP
-		TargetPlatforms.Add(new TGenericMacTargetPlatform<false, false, false>());
-		// Editor TP
-		TargetPlatforms.Add(new TGenericMacTargetPlatform<true, false, false>());
-		// Server TP
-		TargetPlatforms.Add(new TGenericMacTargetPlatform<false, true, false>());
-		// Client TP
-		TargetPlatforms.Add(new TGenericMacTargetPlatform<false, false, true>());
-
-#if PLATFORM_WINDOWS
-		// we added Mac to Windows so that the Xcode Project settings show up, but we don't
-		// want to see the Mac in the Platforms dropdown
-		FDataDrivenPlatformInfoRegistry::SetPlatformHiddenFromUI("Mac");
-#endif
 	}
 
+	virtual void GetTargetPlatforms(TArray<ITargetPlatform*>& TargetPlatforms, TArray<ITargetPlatformSettings*> TargetPlatformSettings, TArray<ITargetPlatformControls*> TargetPlatformControls) override
+	{
+		for (ITargetPlatformControls* TargetPlatformControlsIt : TargetPlatformControls)
+		{
+			TargetPlatforms.Add(new FTargetPlatformMerged(TargetPlatformControlsIt->GetTargetPlatformSettings(), TargetPlatformControlsIt));
+		}
+	}
 
 public:
 
@@ -46,89 +33,13 @@ public:
 
 	virtual void StartupModule() override
 	{
-#if WITH_ENGINE
-		TargetSettings = NewObject<UMacTargetSettings>(GetTransientPackage(), "MacTargetSettings", RF_Standalone);
-		
-		// We need to manually load the config properties here, as this module is loaded before the UObject system is setup to do this
-        GConfig->GetArray(TEXT("/Script/MacTargetPlatform.MacTargetSettings"), TEXT("TargetedRHIs"), TargetSettings->TargetedRHIs, GEngineIni);
-       
-        if (!GConfig->GetInt(TEXT("/Script/MacTargetPlatform.MacTargetSettings"), TEXT("MetalLanguageVersion"), TargetSettings->MetalLanguageVersion, GEngineIni))
-        {
-            TargetSettings->MetalLanguageVersion = 0;
-        }
-        
-		if (!GConfig->GetBool(TEXT("/Script/MacTargetPlatform.MacTargetSettings"), TEXT("UseFastIntrinsics"), TargetSettings->UseFastIntrinsics, GEngineIni))
-		{
-			TargetSettings->UseFastIntrinsics = false;
-		}
-		
-		if (!GConfig->GetBool(TEXT("/Script/MacTargetPlatform.MacTargetSettings"), TEXT("EnableMathOptimisations"), TargetSettings->EnableMathOptimisations, GEngineIni))
-		{
-			TargetSettings->EnableMathOptimisations = true;
-		}
-		
-		if (!GConfig->GetInt(TEXT("/Script/MacTargetPlatform.MacTargetSettings"), TEXT("IndirectArgumentTier"), TargetSettings->IndirectArgumentTier, GEngineIni))
-		{
-			TargetSettings->IndirectArgumentTier = 0;
-		}
-		
-		TargetSettings->AddToRoot();
-        
-        ProjectSettings = NewObject<UXcodeProjectSettings>(GetTransientPackage(), "XcodeProjectSettings", RF_Standalone);
-        ProjectSettings->AddToRoot();
-
-		ISettingsModule* SettingsModule = FModuleManager::GetModulePtr<ISettingsModule>("Settings");
-
-		if (SettingsModule != nullptr)
-		{
-			SettingsModule->RegisterSettings("Project", "Platforms", "Mac",
-				LOCTEXT("MacTargetSettingsName", "Mac"),
-				LOCTEXT("MacTargetSettingsDescription", "Settings and resources for Mac platform"),
-				TargetSettings
-			);
-            SettingsModule->RegisterSettings("Project", "Platforms", "Xcode",
-                LOCTEXT("XcodeProjectSettingsName", "Xcode Projects"),
-                LOCTEXT("XcodeProjectSettingsDescription", "Settings for Xcode projects"),
-                ProjectSettings
-            );
-		}
-#endif
 	}
 
 	virtual void ShutdownModule() override
 	{
-#if WITH_ENGINE
-		ISettingsModule* SettingsModule = FModuleManager::GetModulePtr<ISettingsModule>("Settings");
-
-		if (SettingsModule != nullptr)
-		{
-			SettingsModule->UnregisterSettings("Project", "Platforms", "Mac");
-            SettingsModule->UnregisterSettings("Project", "Platforms", "Xcode");
-		}
-
-		if (!GExitPurge)
-		{
-			// If we're in exit purge, this object has already been destroyed
-			TargetSettings->RemoveFromRoot();
-            ProjectSettings->RemoveFromRoot();
-		}
-		else
-		{
-			TargetSettings = NULL;
-            ProjectSettings = NULL;
-		}
-#endif
 	}
 
 	// End IModuleInterface interface
-
-
-private:
-
-	// Holds the target settings.
-	UMacTargetSettings* TargetSettings;
-    
-    UXcodeProjectSettings* ProjectSettings;
 };
 
 
diff --git a/Engine/Source/Developer/Mac/MacTargetPlatform/Public/GenericMacTargetPlatform.h b/Engine/Source/Developer/Mac/MacTargetPlatform/Public/GenericMacTargetPlatform.h
deleted file mode 100644
diff --git a/Engine/Source/Developer/Mac/MacTargetPlatform/Public/LocalMacTargetDevice.h b/Engine/Source/Developer/Mac/MacTargetPlatform/Public/LocalMacTargetDevice.h
deleted file mode 100644
diff --git a/Engine/Source/Developer/Mac/MacTargetPlatformControls/MacTargetPlatformControls.Build.cs b/Engine/Source/Developer/Mac/MacTargetPlatformControls/MacTargetPlatformControls.Build.cs
new file mode 100644
--- /dev/null
+++ b/Engine/Source/Developer/Mac/MacTargetPlatformControls/MacTargetPlatformControls.Build.cs
@@ -0,0 +1,35 @@
+// Copyright Epic Games, Inc. All Rights Reserved.
+
+using UnrealBuildTool;
+
+public class MacTargetPlatformControls : ModuleRules
+{
+	public MacTargetPlatformControls(ReadOnlyTargetRules Target) : base(Target)
+	{
+		// We need a short name here since this can run afoul very easily of the `MAX_PATH` limit when
+		// combined with building other targets that make use of this.
+		ShortName = "MacTPCon";
+
+		PrivateDependencyModuleNames.AddRange(
+			new string[] {
+				"Core",
+				"CoreUObject",
+				"TargetPlatform",
+				"DesktopPlatform",
+				"MacTargetPlatformSettings"
+			}
+		);
+
+		PrivateIncludePathModuleNames.AddRange(
+			new string[] {
+				"Settings"
+			}
+		);
+
+		if (Target.bCompileAgainstEngine)
+		{
+			PrivateDependencyModuleNames.Add("Engine");
+			PrivateIncludePathModuleNames.Add("TextureCompressor");
+		}
+	}
+}
diff --git a/Engine/Source/Developer/Mac/MacTargetPlatformControls/Private/MacTargetPlatformControlsModule.cpp b/Engine/Source/Developer/Mac/MacTargetPlatformControls/Private/MacTargetPlatformControlsModule.cpp
new file mode 100644
--- /dev/null
+++ b/Engine/Source/Developer/Mac/MacTargetPlatformControls/Private/MacTargetPlatformControlsModule.cpp
@@ -0,0 +1,60 @@
+// Copyright Epic Games, Inc. All Rights Reserved.
+
+#include "CoreMinimal.h"
+#include "GenericMacTargetPlatformControls.h"
+#include "Modules/ModuleManager.h"
+#include "Interfaces/ITargetPlatformControlsModule.h"
+#include "Modules/ModuleManager.h"
+#include "IMacTargetPlatformSettingsModule.h"
+
+#define LOCTEXT_NAMESPACE "FMacTargetPlatformControlsModule"
+
+/**
+ * Module for Mac as a target platform controls
+ */
+class FMacTargetPlatformControlsModule
+	: public ITargetPlatformControlsModule
+{
+public:
+	virtual void GetTargetPlatformControls(TArray<ITargetPlatformControls*>& TargetPlatforms, FName& PlatformSettingsModuleName) override
+	{
+		TMap<FString, ITargetPlatformSettings*> PlatformNameToPlatformSettings;
+		IMacTargetPlatformSettingsModule* ModuleSettings = FModuleManager::GetModulePtr<IMacTargetPlatformSettingsModule>(PlatformSettingsModuleName);
+		if (ModuleSettings != nullptr)
+		{
+			TMap<FString, ITargetPlatformSettings*> OutMap;
+			ModuleSettings->GetPlatformSettingsMaps(OutMap);
+			ITargetPlatformControls* GameTP = new TGenericMacTargetPlatformControls<false, false, false>(OutMap[FMacPlatformProperties<false, false, false>::PlatformName()]);
+			ITargetPlatformControls* EditorTP = new TGenericMacTargetPlatformControls<true, false, false>(OutMap[FMacPlatformProperties<true, false, false>::PlatformName()]);
+			ITargetPlatformControls* ServerTP = new TGenericMacTargetPlatformControls<false, true, false>(OutMap[FMacPlatformProperties<false, true, false>::PlatformName()]);
+			ITargetPlatformControls* ClientTP = new TGenericMacTargetPlatformControls<false, false, true>(OutMap[FMacPlatformProperties<false, false, true>::PlatformName()]);
+			TargetPlatforms.Add(GameTP);
+			TargetPlatforms.Add(EditorTP);
+			TargetPlatforms.Add(ServerTP);
+			TargetPlatforms.Add(ClientTP);
+		}
+	}
+
+public:
+
+	// Begin IModuleInterface interface
+
+	virtual void StartupModule() override
+	{
+
+	}
+
+	virtual void ShutdownModule() override
+	{
+
+	}
+
+	// End IModuleInterface interface
+
+};
+
+
+#undef LOCTEXT_NAMESPACE
+
+
+IMPLEMENT_MODULE( FMacTargetPlatformControlsModule, MacTargetPlatformControls);
diff --git a/Engine/Source/Developer/Mac/MacTargetPlatformControls/Public/GenericMacTargetPlatformControls.h b/Engine/Source/Developer/Mac/MacTargetPlatformControls/Public/GenericMacTargetPlatformControls.h
new file mode 100644
--- /dev/null
+++ b/Engine/Source/Developer/Mac/MacTargetPlatformControls/Public/GenericMacTargetPlatformControls.h
@@ -0,0 +1,168 @@
+// Copyright Epic Games, Inc. All Rights Reserved.
+
+/*=============================================================================
+	GenericMacTargetPlatform.h: Declares the TGenericMacTargetPlatformControls class template.
+=============================================================================*/
+
+#pragma once
+
+#include "CoreMinimal.h"
+#include "Common/TargetPlatformControlsBase.h"
+#include "Mac/MacPlatformProperties.h"
+#include "LocalMacTargetDevice.h"
+#include "AnalyticsEventAttribute.h"
+#if WITH_ENGINE
+#include "StaticMeshResources.h"
+#include "TextureResource.h"
+#endif
+#define LOCTEXT_NAMESPACE "TGenericMacTargetPlatformControls"
+
+/**
+ * Template for Mac target platforms controls
+ */
+template<bool HAS_EDITOR_DATA, bool IS_DEDICATED_SERVER, bool IS_CLIENT_ONLY>
+class TGenericMacTargetPlatformControls
+	: public TTargetPlatformControlsBase<FMacPlatformProperties<HAS_EDITOR_DATA, IS_DEDICATED_SERVER, IS_CLIENT_ONLY> >
+{
+public:
+
+	typedef FMacPlatformProperties<HAS_EDITOR_DATA, IS_DEDICATED_SERVER, IS_CLIENT_ONLY> TProperties;
+	typedef TTargetPlatformControlsBase<TProperties> TSuper;
+
+	/**
+	 * Default constructor.
+	 */
+	TGenericMacTargetPlatformControls(ITargetPlatformSettings* TargetPlatformSettings)
+		: TSuper(TargetPlatformSettings)
+	{
+#if PLATFORM_MAC
+		// only add local device if actually running on Mac
+		LocalDevice = MakeShareable(new FLocalMacTargetDevice(*this));
+#endif
+	}
+
+public:
+
+	//~ Begin ITargetPlatform Interface
+
+	virtual void EnableDeviceCheck(bool OnOff) override {}
+
+	virtual void GetAllDevices( TArray<ITargetDevicePtr>& OutDevices ) const override
+	{
+		OutDevices.Reset();
+		if (LocalDevice.IsValid())
+		{
+			OutDevices.Add(LocalDevice);
+		}
+	}
+
+	virtual bool GenerateStreamingInstallManifest(const TMultiMap<FString, int32>& PakchunkMap, const TSet<int32>& PakchunkIndicesInUse) const override
+	{
+		return true;
+	}
+
+	virtual ITargetDevicePtr GetDefaultDevice( ) const override
+	{
+		if (LocalDevice.IsValid())
+		{
+			return LocalDevice;
+		}
+
+		return nullptr;
+	}
+
+	virtual ITargetDevicePtr GetDevice( const FTargetDeviceId& DeviceId )
+	{
+		if (LocalDevice.IsValid() && (DeviceId == LocalDevice->GetId()))
+		{
+			return LocalDevice;
+		}
+
+		return NULL;
+	}
+
+	virtual bool IsRunningPlatform( ) const override
+	{
+		// Must be Mac platform as editor for this to be considered a running platform
+		return PLATFORM_MAC && !UE_SERVER && !UE_GAME && WITH_EDITOR && HAS_EDITOR_DATA;
+	}
+
+
+	virtual void GetPlatformSpecificProjectAnalytics( TArray<FAnalyticsEventAttribute>& AnalyticsParamArray ) const override
+	{
+		TSuper::GetPlatformSpecificProjectAnalytics(AnalyticsParamArray);
+
+		TSuper::AppendAnalyticsEventConfigArray(AnalyticsParamArray, TEXT("/Script/MacTargetPlatform.MacTargetSettings"), TEXT("TargetedRHIs"), GEngineIni);
+	}
+
+#if WITH_ENGINE
+
+	virtual void GetTextureFormats( const UTexture* Texture, TArray< TArray<FName> >& OutFormats) const override
+	{
+		if (!IS_DEDICATED_SERVER)
+		{
+			// just use the standard texture format name for this texture (with DX11 support)
+			GetDefaultTextureFormatNamePerLayer(OutFormats.AddDefaulted_GetRef(), this->GetTargetPlatformSettings(), this, Texture, true, 4, true);
+		}
+	}
+
+	virtual void GetAllTextureFormats(TArray<FName>& OutFormats) const override
+	{
+		if (!IS_DEDICATED_SERVER)
+		{
+			// just use the standard texture format name for this texture (with DX11 support)
+			GetAllDefaultTextureFormats(this->GetTargetPlatformSettings(), OutFormats);
+		}
+	}
+
+	virtual bool SupportsLQCompressionTextureFormat() const override { return false; };
+
+	virtual bool CanSupportRemoteShaderCompile() const override
+	{
+		return true;
+	}
+	
+	virtual void GetShaderCompilerDependencies(TArray<FString>& OutDependencies) const override
+	{
+		FTargetPlatformControlsBase::AddDependencySCArrayHelper(OutDependencies, TEXT("Binaries/ThirdParty/ShaderConductor/Mac/libdxcompiler.dylib"));
+		FTargetPlatformControlsBase::AddDependencySCArrayHelper(OutDependencies, TEXT("Binaries/ThirdParty/ShaderConductor/Mac/libShaderConductor.dylib"));
+	}
+
+#endif //WITH_ENGINE
+
+
+	virtual bool SupportsVariants() const override
+	{
+		return true;
+	}
+
+
+	virtual float GetVariantPriority() const override
+	{
+		return TProperties::GetVariantPriority();
+	}
+
+	static FORCEINLINE bool SupportsRayTracing()
+	{
+		return true;
+	}
+
+	//~ End ITargetPlatform Interface
+
+private:
+
+	// Holds the local device.
+	ITargetDevicePtr LocalDevice;
+
+#if WITH_ENGINE
+	// Holds the texture LOD settings.
+	const UTextureLODSettings* TextureLODSettings;
+
+	// Holds the static mesh LOD settings.
+	FStaticMeshLODSettings StaticMeshLODSettings;
+
+#endif // WITH_ENGINE
+
+};
+
+#undef LOCTEXT_NAMESPACE
diff --git a/Engine/Source/Developer/Mac/MacTargetPlatformControls/Public/LocalMacTargetDevice.h b/Engine/Source/Developer/Mac/MacTargetPlatformControls/Public/LocalMacTargetDevice.h
new file mode 100644
--- /dev/null
+++ b/Engine/Source/Developer/Mac/MacTargetPlatformControls/Public/LocalMacTargetDevice.h
@@ -0,0 +1,144 @@
+// Copyright Epic Games, Inc. All Rights Reserved.
+
+/*=============================================================================
+	LocalMacTargetDevice.h: Declares the TLocalMacTargetDevice class.
+=============================================================================*/
+
+#pragma once
+
+#include "CoreMinimal.h"
+#include "Interfaces/ITargetDevice.h"
+#include "HAL/FileManager.h"
+#include "HAL/PlatformProcess.h"
+
+/**
+ * Template for local Mac target devices.
+ */
+class FLocalMacTargetDevice
+	: public ITargetDevice
+{
+public:
+
+	/**
+	 * Creates and initializes a new device for the specified target platform.
+	 *
+	 * @param InTargetPlatform - The target platform.
+	 */
+	FLocalMacTargetDevice(const ITargetPlatformControls& InTargetPlatformControls)
+		: TargetPlatformControls(InTargetPlatformControls)
+	{ }
+
+
+public:
+
+	virtual bool Connect( ) override
+	{
+		return true;
+	}
+
+	virtual void Disconnect( )
+	{ }
+
+	virtual ETargetDeviceTypes GetDeviceType( ) const override
+	{
+		return ETargetDeviceTypes::Desktop;
+	}
+
+	virtual FTargetDeviceId GetId( ) const override
+	{
+		return FTargetDeviceId(TargetPlatformControls.PlatformName(), GetName());
+	}
+
+	virtual FString GetName( ) const override
+	{
+		return FPlatformProcess::ComputerName();
+	}
+
+	virtual FString GetOperatingSystemName( ) override
+	{
+		return TEXT("macOS");
+	}
+
+	virtual int32 GetProcessSnapshot( TArray<FTargetDeviceProcessInfo>& OutProcessInfos ) override
+	{
+		// @todo Mac: implement process snapshots
+		return 0;
+	}
+
+	virtual const class ITargetPlatformSettings& GetPlatformSettings() const override
+	{
+		return *(TargetPlatformControls.GetTargetPlatformSettings());
+	}
+	virtual const class ITargetPlatformControls& GetPlatformControls() const override
+	{
+		return TargetPlatformControls;
+	}
+
+	virtual bool IsConnected( )
+	{
+		return true;
+	}
+
+	virtual bool IsDefault( ) const override
+	{
+		return true;
+	}
+
+	virtual bool PowerOff( bool Force ) override
+	{
+		return false;
+	}
+
+	virtual bool PowerOn( ) override
+	{
+		return false;
+	}
+
+	virtual bool Reboot( bool bReconnect = false ) override
+	{
+#if PLATFORM_MAC
+		NSAppleScript* Script = [[NSAppleScript alloc] initWithSource:@"tell application \"System Events\" to restart"];
+		NSDictionary* ErrorDict = [NSDictionary dictionary];
+		[Script executeAndReturnError: &ErrorDict];
+#endif
+		return true;
+	}
+
+	virtual bool SupportsFeature( ETargetDeviceFeatures Feature ) const override
+	{
+		switch (Feature)
+		{
+		case ETargetDeviceFeatures::MultiLaunch:
+			return true;
+
+		// @todo Mac: implement process snapshots
+		case ETargetDeviceFeatures::ProcessSnapshot:
+			return false;
+
+		case ETargetDeviceFeatures::Reboot:
+			return true;
+		}
+
+		return false;
+	}
+
+	virtual void SetUserCredentials( const FString& UserName, const FString& UserPassword ) override
+	{
+	}
+
+	virtual bool GetUserCredentials( FString& OutUserName, FString& OutUserPassword ) override
+	{
+		return false;
+	}
+
+	virtual bool TerminateProcess( const int64 ProcessId ) override
+	{
+		return false;
+	}
+
+
+private:
+
+	// Holds a reference to the device's target platform.
+	const ITargetPlatformControls& TargetPlatformControls;
+};
diff --git a/Engine/Source/Developer/Mac/MacTargetPlatformSettings/Classes/MacTargetSettings.h b/Engine/Source/Developer/Mac/MacTargetPlatformSettings/Classes/MacTargetSettings.h
new file mode 100644
--- /dev/null
+++ b/Engine/Source/Developer/Mac/MacTargetPlatformSettings/Classes/MacTargetSettings.h
@@ -0,0 +1,170 @@
+// Copyright Epic Games, Inc. All Rights Reserved.
+
+/*=============================================================================
+	MacTargetSettings.h: Declares the UMacTargetSettings class.
+=============================================================================*/
+
+#pragma once
+
+#include "CoreMinimal.h"
+#include "UObject/Object.h"
+#include "MacTargetSettings.generated.h"
+
+UENUM()
+enum class EMacMetalShaderStandard : uint8
+{
+    /** Metal Shader 2.2 is the minimum*/
+    MacMetalSLStandard_Minimum = 0 UMETA(DisplayName="Minimum, Currently v2.2 (10.15+)"),
+    /** Metal Shaders Compatible With macOS 10.15 or later (std=osx-metal2.2) */
+    MacMetalSLStandard_2_2 = 5 UMETA(DisplayName="Metal v2.2 (10.15+)"),
+    /** Metal Shaders Compatible With macOS 11.0 or later (std=osx-metal2.3) */
+    MacMetalSLStandard_2_3 = 6 UMETA(DisplayName="Metal v2.3 (11.0+)"),
+    /** Metal Shaders Compatible With macOS 12.0 or later (std=osx-metal2.4) */
+    MacMetalSLStandard_2_4 = 7 UMETA(DisplayName="Metal v2.4 (12.0+)"),
+    /** Metal Shaders Compatible With macOS 13.0 or later (std=metal3.0) */
+    MacMetalSLStandard_3_0 = 8 UMETA(DisplayName="Metal v3.0 (13.0+)"),
+
+};
+
+UENUM()
+enum class EMacTargetArchitecture : uint8
+{
+    /** Create packages that can run natively on Intel Macs and under translation on Apple Silicon Macs */
+    MacTargetArchitectureIntel = 0 UMETA(DisplayName="Intel x64"),
+    
+    /** Create Universal packages that run natively on all Macs */
+    MacTargetArchitectureUniversal = 1 UMETA(DisplayName="Universal (Intel & Apple Silicon)"),
+
+	/** Create packages that can run natively on Apple Silicon Macs, and will not run on Intel Macs */
+	MacTargetArchitectureAppleSilicon = 2 UMETA(DisplayName="Apple Silicon"),
+
+	/** Create  packages that match the architecture of the Mac the package is made on */
+	MacTargetArchitectureHost = 3 UMETA(DisplayName="Host (Matches current machine)"),
+};
+
+
+/**
+ * Implements the settings for the Mac target platform.
+ */
+UCLASS(config=Engine, defaultconfig)
+class MACTARGETPLATFORMSETTINGS_API UMacTargetSettings
+	: public UObject
+{
+public:
+
+	GENERATED_UCLASS_BODY()
+	
+	/**
+	 * The collection of RHI's we want to support on this platform.
+	 * This is not always the full list of RHI we can support.
+	 */
+	UPROPERTY(EditAnywhere, config, Category=Rendering)
+	TArray<FString> TargetedRHIs;
+	
+	/**
+	 * The set of architecture(s) this project supports for Editor builds
+	 * This defines which CPU architectures to target: x86_64 (Intel), arm64 (Apple Silicon) or Universal (Intel & Apple Silicon).
+	 * It is recommended to use Universal unless you have editor plugins or other libraries that do not support Apple Silicon
+	 */
+	UPROPERTY(EditAnywhere, config, Category=Packaging, meta = (DisplayName = "Supported Architecture(s) for Editor builds"))
+	EMacTargetArchitecture EditorTargetArchitecture;
+	
+	/**
+	 * The target Mac platform CPU architecture.
+	 * This defines which CPU architectures to target: x86_64 (Intel), arm64 (Apple Silicon) or Universal (Intel & Apple Silicon).
+	 * It is recommended to use Universal unless you have runtime plugins or other libraries that do not support Apple Silicon
+	 */
+	UPROPERTY(EditAnywhere, config, Category=Packaging, meta = (DisplayName = "Supported Architecture(s) for non-Editor builds"))
+	EMacTargetArchitecture TargetArchitecture;
+	
+	/**
+	 * The architecture to compile the Editor target
+	 * This defines which CPU architectures to target: x86_64 (Intel), arm64 (Apple Silicon) or Universal (Intel & Apple Silicon), or Host to match the machine doing the building
+	 * Can override with -Architecture= on the UBT commandline, or -EditorArchitecture=  on the BuildCookRun commandline
+	 */
+	UPROPERTY(EditAnywhere, config, Category=Packaging, meta = (DisplayName = "Architecture(s) when building Editor"))
+	EMacTargetArchitecture EditorDefaultArchitecture;
+	
+	/**
+	 * The architectures to compile non-Editor (games, programs, etc) targets for builds outside of Xcode
+	 * This defines which CPU architectures to target: x86_64 (Intel), arm64 (Apple Silicon) or Universal (Intel & Apple Silicon), or Host to match the machine doing the building
+	 * Can override with -Architecture= on the UBT commandline, or -GameArchitecture= or -ProgramArchitecture= on the BuildCookRun commandline
+	 */
+	UPROPERTY(EditAnywhere, config, Category=Packaging, meta = (DisplayName = "Architecture(s) when building non-Editor"))
+	EMacTargetArchitecture DefaultArchitecture;
+	
+	/**
+	 * If true, builds running on BuildMachines (when the 'IsBuildMachine' environment variable is set to 1) will compile all Supported architectures
+	 */
+	UPROPERTY(EditAnywhere, config, Category=Packaging, meta = (DisplayName = "Build all supported Architectures on Build Machines"))
+	bool bBuildAllSupportedOnBuildMachine;
+
+    /**
+     * The Metal shader language version which will be used when compiling the shaders.
+     */
+    UPROPERTY(EditAnywhere, config, Category=Rendering, meta = (DisplayName = "Metal Shader Standard To Target", ConfigRestartRequired = true))
+    int32 MetalLanguageVersion;
+    
+    /**
+     * Whether to use the Metal shading language's "fast" intrinsics.
+	 * Fast intrinsics assume that no NaN or INF value will be provided as input, 
+	 * so are more efficient. However, they will produce undefined results if NaN/INF 
+	 * is present in the argument/s. By default fast-instrinics are disabled so Metal correctly handles NaN/INF arguments.
+     */
+    UPROPERTY(EditAnywhere, config, Category=Rendering, meta = (DisplayName = "Use Fast-Math intrinsics", ConfigRestartRequired = true))
+	bool UseFastIntrinsics;
+	
+	/**
+	 * Whether to use of Metal shader-compiler's -ffast-math optimisations.
+	 * Fast-Math performs algebraic-equivalent & reassociative optimisations not permitted by the floating point arithmetic standard (IEEE-754).
+	 * These can improve shader performance at some cost to precision and can lead to NaN/INF propagation as they rely on
+	 * shader inputs or variables not containing NaN/INF values. By default fast-math is enabled for performance.
+	 */
+	UPROPERTY(EditAnywhere, config, Category=Rendering, meta = (DisplayName = "Enable Fast-Math optimisations", ConfigRestartRequired = true))
+	bool EnableMathOptimisations;
+	
+	/** Whether to compile shaders using a tier Indirect Argument Buffers. */
+	UPROPERTY(config, EditAnywhere, Category = Rendering, Meta = (DisplayName = "Tier of Indirect Argument Buffers to use when compiling shaders", ConfigRestartRequired = true))
+	int32 IndirectArgumentTier;
+
+	/** Sample rate to run the audio mixer with. */
+	UPROPERTY(config, EditAnywhere, Category = "Audio", Meta = (DisplayName = "Audio Mixer Sample Rate"))
+	int32 AudioSampleRate;
+
+	/** The amount of audio to compute each callback block. Lower values decrease latency but may increase CPU cost. */
+	UPROPERTY(config, EditAnywhere, Category = "Audio", meta = (ClampMin = "512", ClampMax = "4096", DisplayName = "Callback Buffer Size"))
+	int32 AudioCallbackBufferFrameSize;
+
+	/** The number of buffers to keep enqueued. More buffers increases latency, but can compensate for variable compute availability in audio callbacks on some platforms. */
+	UPROPERTY(config, EditAnywhere, Category = "Audio", meta = (ClampMin = "1", UIMin = "1", DisplayName = "Number of Buffers To Enqueue"))
+	int32 AudioNumBuffersToEnqueue;
+
+	/** The max number of channels (voices) to limit for this platform. The max channels used will be the minimum of this value and the global audio quality settings. A value of 0 will not apply a platform channel count max. */
+	UPROPERTY(config, EditAnywhere, Category = "Audio", meta = (ClampMin = "0", UIMin = "0", DisplayName = "Max Channels"))
+	int32 AudioMaxChannels;
+
+	/** The number of workers to use to compute source audio. Will only use up to the max number of sources. Will evenly divide sources to each source worker. */
+	UPROPERTY(config, EditAnywhere, Category = "Audio", meta = (ClampMin = "0", UIMin = "0", DisplayName = "Number of Source Workers"))
+	int32 AudioNumSourceWorkers;
+
+	/** Which of the currently enabled spatialization plugins to use. */
+	UPROPERTY(config, EditAnywhere, Category = "Audio")
+	FString SpatializationPlugin;
+
+	/** Which of the currently enabled source data override plugins to use. */
+	UPROPERTY(config, EditAnywhere, Category = "Audio")
+	FString SourceDataOverridePlugin;
+
+	/** Which of the currently enabled reverb plugins to use. */
+	UPROPERTY(config, EditAnywhere, Category = "Audio")
+	FString ReverbPlugin;
+
+	/** Which of the currently enabled occlusion plugins to use. */
+	UPROPERTY(config, EditAnywhere, Category = "Audio")
+	FString OcclusionPlugin;
+	
+	/** Quality Level to COOK SoundCues at (if set, all other levels will be stripped by the cooker). */
+	UPROPERTY(GlobalConfig, EditAnywhere, Category = "Audio|CookOverrides", meta = (DisplayName = "Sound Cue Cook Quality"))
+	int32 SoundCueCookQualityIndex = INDEX_NONE;
+
+};
diff --git a/Engine/Source/Developer/Mac/MacTargetPlatformSettings/Classes/XcodeProjectSettings.h b/Engine/Source/Developer/Mac/MacTargetPlatformSettings/Classes/XcodeProjectSettings.h
new file mode 100644
--- /dev/null
+++ b/Engine/Source/Developer/Mac/MacTargetPlatformSettings/Classes/XcodeProjectSettings.h
@@ -0,0 +1,178 @@
+// Copyright Epic Games, Inc. All Rights Reserved.
+
+/*=============================================================================
+    XcodeProjectSettings.h: Declares the UXcodeProjectSettings class.
+=============================================================================*/
+
+#pragma once
+
+#include "CoreMinimal.h"
+#include "UObject/Object.h"
+#include "XcodeProjectSettings.generated.h"
+
+/**
+ * Implements the settings for Xcode projects
+ */
+UCLASS(config=Engine, defaultconfig)
+class MACTARGETPLATFORMSETTINGS_API UXcodeProjectSettings
+    : public UObject
+{
+public:
+	
+	GENERATED_UCLASS_BODY()
+	
+	/**
+	 * Enable modernized Xcode, when building from Xcode, use native Xcode for bundle generation and archiving instead of UBT
+     * Restart required to apply this setting
+	 */
+	UPROPERTY(EditAnywhere, config, Category=Xcode, meta = (DisplayName = "Modernized Xcode", ConfigRestartRequired = true))
+	bool bUseModernXcode;
+    
+    UFUNCTION()
+    static bool ShouldDisableIOSSettings()
+    {
+#if PLATFORM_MAC
+        auto DefaultObject = Cast<UXcodeProjectSettings>(UXcodeProjectSettings::StaticClass()->GetDefaultObject());
+        if (DefaultObject)
+        {
+            // Some iOS build settings no longer functional under Modern Xcode
+            return DefaultObject->bUseModernXcode;
+        }
+        else
+        {
+            return false;
+        }
+#else
+        // On PC, even after turning on Modern Xcode for remote Mac, should let user config local iOS related settings
+        return false;
+#endif
+    }
+	
+	/**
+	 * Team ID used for native Xcode code signing. This must be the 10 letters/numbers ID found in Membership Details tab found in https://developer.apple.com/account
+	 */
+	UPROPERTY(EditAnywhere, config, Category=Xcode, meta = (EditCondition="bUseModernXcode", DisplayName = "Apple Dev Account Team ID"))
+	FString CodeSigningTeam;
+	
+	/**
+	 * Bundle ID used for nativr Xcode code signing
+	 */
+	UPROPERTY(EditAnywhere, config, Category=Xcode, meta = (EditCondition="bUseModernXcode", DisplayName = "Bundle ID"))
+	FString BundleIdentifier;
+	
+	/**
+	 * Bundle ID prefix used for native Xcode code signing. This is only needed if you use the default, pieced-together Bundle ID above. If you specify a full Bundle ID, you can ignore this field.
+	 */
+	UPROPERTY(EditAnywhere, config, Category=Xcode, meta = (EditCondition="bUseModernXcode", DisplayName = "Bundle ID Prefix"))
+	FString CodeSigningPrefix;
+	
+	/**
+	 * The name of the Mac .app when making an archived build (for uploading to App Store, etc). The Finder shows Mac apps by their .app name, and we don't name the .app  with
+	 * "pretty" names during development. When packaging for distribution (or using Archive menu in Xcode) this will become the name of the .app, and will be what end users
+	 * will have on their Mac. If this is not set, the .app will have the name of the .uproject file
+	 */
+	UPROPERTY(EditAnywhere, config, Category=Xcode, meta = (EditCondition="bUseModernXcode", DisplayName = "Mac: Published App Name"))
+	FString ApplicationDisplayName;
+
+	/**
+	 * The App Category that will be used for Apple App Store submissions
+	 */
+	UPROPERTY(EditAnywhere, config, Category=Xcode, meta = (EditCondition="bUseModernXcode", DisplayName = "App Category"))
+	FString AppCategory;
+	
+	/**
+	 * The template info.plist used for Mac game targets
+	 */
+	UPROPERTY(EditAnywhere, config, Category="Plist Files", meta = (EditCondition="bUseModernXcode", DisplayName = "Mac: Info.plist Template"))
+	FFilePath TemplateMacPlist;
+	
+	/**
+	 * The template info.plist used for iOS game targets
+	 */
+	UPROPERTY(EditAnywhere, config, Category="Plist Files", meta = (EditCondition="bUseModernXcode", DisplayName = "IOS / TVOS: Info.plist Template"))
+	FFilePath TemplateIOSPlist;
+	
+	/**
+	 * The premade entitlement file used for development Mac builds
+	 */
+	UPROPERTY(EditAnywhere, config, Category="Entitlements", meta = (EditCondition="bUseModernXcode", DisplayName = "Mac: Development Entitlements"))
+	FFilePath PremadeMacEntitlements;
+	
+	/**
+	 * The premade entitlement file used for shipping Mac builds
+	 */
+	UPROPERTY(EditAnywhere, config, Category="Entitlements", meta = (EditCondition="bUseModernXcode", DisplayName = "Mac: Shipping Entitlements"))
+	FFilePath ShippingSpecificMacEntitlements;
+	
+	/**
+	 * Enable native Xcode code signing
+	 */
+	UPROPERTY(EditAnywhere, config, Category="Code Signing", meta = (EditCondition="bUseModernXcode"))
+	bool bUseAutomaticCodeSigning;
+	
+	/**
+	 * If true, Mac will sign to run locally. Running on another Mac may bring up a dialog preventing running the app. If this and Use Automatic Code Signing are both false, you will need a certificate installed
+	 */
+	UPROPERTY(EditAnywhere, config, Category="Code Signing", meta = (EditCondition="bUseModernXcode"), DisplayName="Mac: Sign To Run Locally")
+	bool bMacSignToRunLocally;
+	
+	/**
+	 * The name (prefix or full) of the certificate to use for Mac code signing.
+	 */
+	UPROPERTY(EditAnywhere, config, Category="Code Signing", meta = (EditCondition="bUseModernXcode && !bUseAutomaticCodeSigning && !bMacSignToRunLocally"), DisplayName = "Mac: Signing Identity")
+	FString MacSigningIdentity;
+	
+	/**
+	 * The name (prefix or full) of the certificate to use for IOS and TVOS code signing
+	 */
+	UPROPERTY(EditAnywhere, config, Category="Code Signing", meta = (EditCondition="bUseModernXcode && !bUseAutomaticCodeSigning", DisplayName = "IOS / TVOS: Signing Identity"))
+	FString IOSSigningIdentity;
+	
+	/**
+	 * The path to a .mobileprovision file to use for signing for IOS. Alternatively, if it's a single name or UUID (no .mobileprovision extension), it will use this as the name/UUID of an already installed provision to sign with
+	 */
+	UPROPERTY(EditAnywhere, config, Category="Code Signing", meta = (EditCondition="bUseModernXcode && !bUseAutomaticCodeSigning", DisplayName = "IOS: Provisioning Profile"))
+	FFilePath IOSProvisioningProfile;
+	
+	/**
+	 * The path to a .mobileprovision file to use for signing for TVOS. Alternatively, if it's a single name or UUID (no .mobileprovision extension), it will use this as the name/UUID of an already installed provision to sign with
+	 */
+	UPROPERTY(EditAnywhere, config, Category="Code Signing", meta = (EditCondition="bUseModernXcode && !bUseAutomaticCodeSigning", DisplayName = "TVOS: Provisioning Profile"))
+	FFilePath TVOSProvisioningProfile;
+	
+	/**
+	 * If true, use AppStore Connect authentication for commandline builds. This allows for automatic codesigning functionality without needing to be signed in to Xcode. See the App Store Connect API section of the Keys tab in your Users and Access page in Apple dev center.
+	 */
+	UPROPERTY(EditAnywhere, config, Category="Code Signing", meta = (EditCondition="bUseModernXcode && bUseAutomaticCodeSigning"))
+	bool bUseAppStoreConnect;
+	
+	/**
+	 * The Issuer ID for your App Store Connect API
+	 */
+	UPROPERTY(EditAnywhere, config, Category="Code Signing", meta = (EditCondition="bUseModernXcode && bUseAutomaticCodeSigning && bUseAppStoreConnect", DisplayName = "App Store Connect: Issuer ID"))
+	FString AppStoreConnectIssuerID;
+	
+	/**
+	 * The Key ID for your App Store Connect generated API key, a 32 hex-character string, including dashes
+	 */
+	UPROPERTY(EditAnywhere, config, Category="Code Signing", meta = (EditCondition="bUseModernXcode && bUseAutomaticCodeSigning && bUseAppStoreConnect", DisplayName = "App Store Connect: Key ID"))
+	FString AppStoreConnectKeyID;
+	
+	/**
+	 * The path to the downloaded .p8 file shared with your team
+	 */
+	UPROPERTY(EditAnywhere, config, Category="Code Signing", meta = (EditCondition="bUseModernXcode && bUseAutomaticCodeSigning && bUseAppStoreConnect", DisplayName = "App Store Connect: Key File"))
+	FFilePath AppStoreConnectKeyPath;
+	
+	/**
+	 * The path to the optional PrivacyInfo.xcprivacy file for your project
+	 */
+	UPROPERTY(EditAnywhere, config, Category="Privacy Manifests", meta = (EditCondition="bUseModernXcode", DisplayName = "Mac: Additional Privacy Info"))
+	FFilePath AdditionalPrivacyInfoMac;
+	
+	/**
+	 * The path to the optional PrivacyInfo.xcprivacy file for your project
+	 */
+	UPROPERTY(EditAnywhere, config, Category="Privacy Manifests", meta = (EditCondition="bUseModernXcode", DisplayName = "IOS/TVOS/VisionOS: Additional Privacy Info"))
+	FFilePath AdditionalPrivacyInfoIOS;
+};
diff --git a/Engine/Source/Developer/Mac/MacTargetPlatformSettings/MacTargetPlatformSettings.Build.cs b/Engine/Source/Developer/Mac/MacTargetPlatformSettings/MacTargetPlatformSettings.Build.cs
new file mode 100644
--- /dev/null
+++ b/Engine/Source/Developer/Mac/MacTargetPlatformSettings/MacTargetPlatformSettings.Build.cs
@@ -0,0 +1,34 @@
+// Copyright Epic Games, Inc. All Rights Reserved.
+
+using UnrealBuildTool;
+
+public class MacTargetPlatformSettings : ModuleRules
+{
+	public MacTargetPlatformSettings(ReadOnlyTargetRules Target) : base(Target)
+	{
+		// We need a short name here since this can run afoul very easily of the `MAX_PATH` limit when
+		// combined with building other targets that make use of this.
+		ShortName = "MacTPSet";
+
+		PrivateDependencyModuleNames.AddRange(
+			new string[] {
+				"Core",
+				"CoreUObject",
+				"TargetPlatform",
+				"DesktopPlatform"
+			}
+		);
+
+		PrivateIncludePathModuleNames.AddRange(
+			new string[] {
+				"Settings"
+			}
+		);
+
+		if (Target.bCompileAgainstEngine)
+		{
+			PrivateDependencyModuleNames.Add("Engine");
+			PrivateIncludePathModuleNames.Add("TextureCompressor");
+		}
+	}
+}
diff --git a/Engine/Source/Developer/Mac/MacTargetPlatformSettings/Private/MacTargetPlatformClasses.cpp b/Engine/Source/Developer/Mac/MacTargetPlatformSettings/Private/MacTargetPlatformClasses.cpp
new file mode 100644
--- /dev/null
+++ b/Engine/Source/Developer/Mac/MacTargetPlatformSettings/Private/MacTargetPlatformClasses.cpp
@@ -0,0 +1,23 @@
+// Copyright Epic Games, Inc. All Rights Reserved.
+
+/*=============================================================================
+	MacTargetPlatformClasses.cpp: Implements the module's UClasses.
+=============================================================================*/
+
+#include "MacTargetSettings.h"
+#include "XcodeProjectSettings.h"
+
+
+/* UMacTargetSettings structors
+ *****************************************************************************/
+
+UMacTargetSettings::UMacTargetSettings( const FObjectInitializer& ObjectInitializer )
+	: Super(ObjectInitializer)
+{ }
+
+/* UXcodeProjectSettings structors
+ *****************************************************************************/
+
+UXcodeProjectSettings::UXcodeProjectSettings( const FObjectInitializer& ObjectInitializer )
+    : Super(ObjectInitializer)
+{ }
diff --git a/Engine/Source/Developer/Mac/MacTargetPlatformSettings/Private/MacTargetPlatformSettingsModule.cpp b/Engine/Source/Developer/Mac/MacTargetPlatformSettings/Private/MacTargetPlatformSettingsModule.cpp
new file mode 100644
--- /dev/null
+++ b/Engine/Source/Developer/Mac/MacTargetPlatformSettings/Private/MacTargetPlatformSettingsModule.cpp
@@ -0,0 +1,145 @@
+// Copyright Epic Games, Inc. All Rights Reserved.
+
+#include "CoreMinimal.h"
+#include "GenericMacTargetPlatformSettings.h"
+#include "Modules/ModuleManager.h"
+#include "ISettingsModule.h"
+#include "Interfaces/ITargetPlatformSettingsModule.h"
+#include "Modules/ModuleManager.h"
+#include "MacTargetSettings.h"
+#include "XcodeProjectSettings.h"
+#include "UObject/Package.h"
+#include "UObject/WeakObjectPtr.h"
+#include "IMacTargetPlatformSettingsModule.h"
+
+#define LOCTEXT_NAMESPACE "FMacTargetPlatformSettingsModule"
+
+/**
+ * Module for Mac as a target platform settings
+ */
+class FMacTargetPlatformSettingsModule
+	: public IMacTargetPlatformSettingsModule
+{
+public:
+	virtual void GetTargetPlatformSettings(TArray<ITargetPlatformSettings*>& TargetPlatforms) override
+	{
+		// Game TP
+		ITargetPlatformSettings* GameTP = new TGenericMacTargetPlatformSettings<false, false, false>();
+		TargetPlatforms.Add(GameTP);
+		PlatformNameToPlatformSettings.Add(FMacPlatformProperties<false, false, false>::PlatformName(), GameTP);
+		// Editor TP
+		ITargetPlatformSettings* EditorTP = new TGenericMacTargetPlatformSettings<true, false, false>();
+		TargetPlatforms.Add(EditorTP);
+		PlatformNameToPlatformSettings.Add(FMacPlatformProperties<true, false, false>::PlatformName(), EditorTP);
+		// Server TP
+		ITargetPlatformSettings* ServerTP = new TGenericMacTargetPlatformSettings<false, true, false>();
+		TargetPlatforms.Add(ServerTP);
+		PlatformNameToPlatformSettings.Add(FMacPlatformProperties<false, true, false>::PlatformName(), ServerTP);
+		// Client TP
+		ITargetPlatformSettings* ClientTP = new TGenericMacTargetPlatformSettings<false, false, true>();
+		TargetPlatforms.Add(ClientTP);
+		PlatformNameToPlatformSettings.Add(FMacPlatformProperties<false, false, true>::PlatformName(), ClientTP);
+	}
+
+	virtual void GetPlatformSettingsMaps(TMap<FString, ITargetPlatformSettings*>& OutMap) override
+	{
+		OutMap = PlatformNameToPlatformSettings;
+	}
+
+public:
+
+	// Begin IModuleInterface interface
+
+	virtual void StartupModule() override
+	{
+#if WITH_ENGINE
+		TargetSettings = NewObject<UMacTargetSettings>(GetTransientPackage(), "MacTargetSettings", RF_Standalone);
+		
+		// We need to manually load the config properties here, as this module is loaded before the UObject system is setup to do this
+        GConfig->GetArray(TEXT("/Script/MacTargetPlatform.MacTargetSettings"), TEXT("TargetedRHIs"), TargetSettings->TargetedRHIs, GEngineIni);
+       
+        if (!GConfig->GetInt(TEXT("/Script/MacTargetPlatform.MacTargetSettings"), TEXT("MetalLanguageVersion"), TargetSettings->MetalLanguageVersion, GEngineIni))
+        {
+            TargetSettings->MetalLanguageVersion = 0;
+        }
+        
+		if (!GConfig->GetBool(TEXT("/Script/MacTargetPlatform.MacTargetSettings"), TEXT("UseFastIntrinsics"), TargetSettings->UseFastIntrinsics, GEngineIni))
+		{
+			TargetSettings->UseFastIntrinsics = false;
+		}
+		
+		if (!GConfig->GetBool(TEXT("/Script/MacTargetPlatform.MacTargetSettings"), TEXT("EnableMathOptimisations"), TargetSettings->EnableMathOptimisations, GEngineIni))
+		{
+			TargetSettings->EnableMathOptimisations = true;
+		}
+		
+		if (!GConfig->GetInt(TEXT("/Script/MacTargetPlatform.MacTargetSettings"), TEXT("IndirectArgumentTier"), TargetSettings->IndirectArgumentTier, GEngineIni))
+		{
+			TargetSettings->IndirectArgumentTier = 0;
+		}
+		
+		TargetSettings->AddToRoot();
+        
+        ProjectSettings = NewObject<UXcodeProjectSettings>(GetTransientPackage(), "XcodeProjectSettings", RF_Standalone);
+        ProjectSettings->AddToRoot();
+
+		ISettingsModule* SettingsModule = FModuleManager::GetModulePtr<ISettingsModule>("Settings");
+
+		if (SettingsModule != nullptr)
+		{
+			SettingsModule->RegisterSettings("Project", "Platforms", "Mac",
+				LOCTEXT("MacTargetSettingsName", "Mac"),
+				LOCTEXT("MacTargetSettingsDescription", "Settings and resources for Mac platform"),
+				TargetSettings
+			);
+            SettingsModule->RegisterSettings("Project", "Platforms", "Xcode",
+                LOCTEXT("XcodeProjectSettingsName", "Xcode Projects"),
+                LOCTEXT("XcodeProjectSettingsDescription", "Settings for Xcode projects"),
+                ProjectSettings
+            );
+		}
+#endif
+	}
+
+	virtual void ShutdownModule() override
+	{
+#if WITH_ENGINE
+		ISettingsModule* SettingsModule = FModuleManager::GetModulePtr<ISettingsModule>("Settings");
+
+		if (SettingsModule != nullptr)
+		{
+			SettingsModule->UnregisterSettings("Project", "Platforms", "Mac");
+            SettingsModule->UnregisterSettings("Project", "Platforms", "Xcode");
+		}
+
+		if (!GExitPurge)
+		{
+			// If we're in exit purge, this object has already been destroyed
+			TargetSettings->RemoveFromRoot();
+            ProjectSettings->RemoveFromRoot();
+		}
+		else
+		{
+			TargetSettings = NULL;
+            ProjectSettings = NULL;
+		}
+#endif
+	}
+
+	// End IModuleInterface interface
+
+
+private:
+	TMap<FString, ITargetPlatformSettings*> PlatformNameToPlatformSettings;
+
+	// Holds the target settings.
+	UMacTargetSettings* TargetSettings;
+    
+    UXcodeProjectSettings* ProjectSettings;
+};
+
+
+#undef LOCTEXT_NAMESPACE
+
+
+IMPLEMENT_MODULE( FMacTargetPlatformSettingsModule, MacTargetPlatformSettings);
diff --git a/Engine/Source/Developer/Mac/MacTargetPlatformSettings/Public/GenericMacTargetPlatformSettings.h b/Engine/Source/Developer/Mac/MacTargetPlatformSettings/Public/GenericMacTargetPlatformSettings.h
new file mode 100644
--- /dev/null
+++ b/Engine/Source/Developer/Mac/MacTargetPlatformSettings/Public/GenericMacTargetPlatformSettings.h
@@ -0,0 +1,149 @@
+// Copyright Epic Games, Inc. All Rights Reserved.
+
+/*=============================================================================
+	GenericMacTargetPlatform.h: Declares the TGenericMacTargetPlatformSettings class template.
+=============================================================================*/
+
+#pragma once
+
+#include "CoreMinimal.h"
+#include "Common/TargetPlatformSettingsBase.h"
+#include "Mac/MacPlatformProperties.h"
+
+#if WITH_ENGINE
+#include "TextureResource.h"
+#include "StaticMeshResources.h"
+#endif // WITH_ENGINE
+
+#define LOCTEXT_NAMESPACE "TGenericMacTargetPlatformSettings"
+
+/**
+ * Template for Mac target platforms settings
+ */
+template<bool HAS_EDITOR_DATA, bool IS_DEDICATED_SERVER, bool IS_CLIENT_ONLY>
+class TGenericMacTargetPlatformSettings
+	: public TTargetPlatformSettingsBase<FMacPlatformProperties<HAS_EDITOR_DATA, IS_DEDICATED_SERVER, IS_CLIENT_ONLY> >
+{
+public:
+
+	typedef FMacPlatformProperties<HAS_EDITOR_DATA, IS_DEDICATED_SERVER, IS_CLIENT_ONLY> TProperties;
+	typedef TTargetPlatformSettingsBase<TProperties> TSuper;
+
+	/**
+	 * Default constructor.
+	 */
+	TGenericMacTargetPlatformSettings( )
+	{
+		#if WITH_ENGINE
+			TextureLODSettings = nullptr;
+			StaticMeshLODSettings.Initialize(this);
+		#endif
+	}
+
+public:
+
+	//~ Begin ITargetPlatform Interface
+
+	virtual bool SupportsFeature( ETargetPlatformFeatures Feature ) const override
+	{
+		// we currently do not have a build target for MacServer
+		if (Feature == ETargetPlatformFeatures::Packaging)
+		{
+			return (HAS_EDITOR_DATA || !IS_DEDICATED_SERVER);
+		}
+
+		if (Feature == ETargetPlatformFeatures::CanCookPackages)
+		{
+			return HAS_EDITOR_DATA;
+		}
+
+		return TSuper::SupportsFeature(Feature);
+	}
+
+	virtual void GetAllPossibleShaderFormats( TArray<FName>& OutFormats ) const override
+	{
+		// no shaders needed for dedicated server target
+		if (!IS_DEDICATED_SERVER)
+		{
+			static FName NAME_SF_METAL_SM5(TEXT("SF_METAL_SM5"));
+			OutFormats.AddUnique(NAME_SF_METAL_SM5);
+            static FName NAME_SF_METAL_SM6(TEXT("SF_METAL_SM6"));
+            OutFormats.AddUnique(NAME_SF_METAL_SM6);
+			static FName NAME_SF_METAL_MACES3_1(TEXT("SF_METAL_MACES3_1"));
+			OutFormats.AddUnique(NAME_SF_METAL_MACES3_1);
+			static FName NAME_SF_METAL_MRT_MAC(TEXT("SF_METAL_MRT_MAC"));
+			OutFormats.AddUnique(NAME_SF_METAL_MRT_MAC);
+		}
+	}
+
+	virtual void GetAllTargetedShaderFormats(TArray<FName>& OutFormats) const override
+	{
+		// Get the Target RHIs for this platform, we do not always want all those that are supported.
+		TArray<FString>TargetedShaderFormats;
+		GConfig->GetArray(TEXT("/Script/MacTargetPlatform.MacTargetSettings"), TEXT("TargetedRHIs"), TargetedShaderFormats, GEngineIni);
+
+		// Gather the list of Target RHIs and filter out any that may be invalid.
+		TArray<FName> PossibleShaderFormats;
+		GetAllPossibleShaderFormats(PossibleShaderFormats);
+
+		for (int32 ShaderFormatIdx = TargetedShaderFormats.Num() - 1; ShaderFormatIdx >= 0; ShaderFormatIdx--)
+		{
+			FString ShaderFormat = TargetedShaderFormats[ShaderFormatIdx];
+			if (PossibleShaderFormats.Contains(FName(*ShaderFormat)) == false)
+			{
+				TargetedShaderFormats.RemoveAt(ShaderFormatIdx);
+			}
+		}
+
+		for(const FString& ShaderFormat : TargetedShaderFormats)
+		{
+			OutFormats.AddUnique(FName(*ShaderFormat));
+		}
+	}
+
+#if WITH_ENGINE
+	virtual const class FStaticMeshLODSettings& GetStaticMeshLODSettings( ) const override
+	{
+		return StaticMeshLODSettings;
+	}
+
+	virtual const UTextureLODSettings& GetTextureLODSettings() const override
+	{
+		return *TextureLODSettings;
+	}
+
+	virtual void RegisterTextureLODSettings(const UTextureLODSettings* InTextureLODSettings) override
+	{
+		TextureLODSettings = InTextureLODSettings;
+	}
+
+#endif //WITH_ENGINE
+
+	static FORCEINLINE bool SupportsRayTracing()
+	{
+		return true;
+	}
+
+	virtual bool UsesRayTracing() const override
+	{
+		bool bEnableRayTracing = false;
+		GConfig->GetBool(TEXT("/Script/MacTargetPlatform.MacTargetSettings"), TEXT("bEnableRayTracing"), bEnableRayTracing, GEngineIni);
+	 
+		return bEnableRayTracing;
+	}
+	//~ End ITargetPlatform Interface
+
+private:
+
+#if WITH_ENGINE
+	// Holds the texture LOD settings.
+	const UTextureLODSettings* TextureLODSettings;
+
+	// Holds the static mesh LOD settings.
+	FStaticMeshLODSettings StaticMeshLODSettings;
+
+#endif // WITH_ENGINE
+
+};
+
+#undef LOCTEXT_NAMESPACE
diff --git a/Engine/Source/Developer/Mac/MacTargetPlatformSettings/Public/IMacTargetPlatformSettingsModule.h b/Engine/Source/Developer/Mac/MacTargetPlatformSettings/Public/IMacTargetPlatformSettingsModule.h
new file mode 100644
--- /dev/null
+++ b/Engine/Source/Developer/Mac/MacTargetPlatformSettings/Public/IMacTargetPlatformSettingsModule.h
@@ -0,0 +1,16 @@
+// Copyright Epic Games, Inc. All Rights Reserved.
+
+#pragma  once 
+
+/*------------------------------------------------------------------------------------
+IMacTargetPlatformSettingsModule interface
+------------------------------------------------------------------------------------*/
+
+#include "CoreMinimal.h"
+#include "Interfaces/ITargetPlatformSettingsModule.h"
+
+class IMacTargetPlatformSettingsModule : public ITargetPlatformSettingsModule
+{
+public:
+	virtual void GetPlatformSettingsMaps(TMap<FString, ITargetPlatformSettings*>& OutMap) = 0;
+};
diff --git a/Engine/Source/Developer/Windows/WindowsPlatformEditor/WindowsPlatformEditor.Build.cs b/Engine/Source/Developer/Windows/WindowsPlatformEditor/WindowsPlatformEditor.Build.cs
--- a/Engine/Source/Developer/Windows/WindowsPlatformEditor/WindowsPlatformEditor.Build.cs
+++ b/Engine/Source/Developer/Windows/WindowsPlatformEditor/WindowsPlatformEditor.Build.cs
@@ -11,7 +11,7 @@ public class WindowsPlatformEditor : ModuleRules
 				"Core",
 				"CoreUObject",
 				
-				"WindowsTargetPlatform",
+				"WindowsTargetPlatformSettings",
 			}
 		);
 
diff --git a/Engine/Source/Developer/Windows/WindowsTargetPlatfomrControls/Private/WindowsTargetPlatformControlsModule.cpp b/Engine/Source/Developer/Windows/WindowsTargetPlatfomrControls/Private/WindowsTargetPlatformControlsModule.cpp
new file mode 100644
--- /dev/null
+++ b/Engine/Source/Developer/Windows/WindowsTargetPlatfomrControls/Private/WindowsTargetPlatformControlsModule.cpp
@@ -0,0 +1,70 @@
+// Copyright Epic Games, Inc. All Rights Reserved.
+
+#include "CoreMinimal.h"
+#include "ISettingsModule.h"
+#include "Interfaces/ITargetPlatformControlsModule.h"
+#include "Interfaces/ITargetPlatformSettingsModule.h"
+#include "GenericWindowsTargetPlatformControls.h"
+#include "Modules/ModuleManager.h"
+#include "UObject/Package.h"
+#include "UObject/WeakObjectPtr.h"
+#if WITH_ENGINE
+#include "CookedEditorTargetPlatformControls.h"
+#endif
+#include "IWindowsTargetPlatformSettingsModule.h"
+
+#define LOCTEXT_NAMESPACE "FWindowsTargetPlatformControlsModule"
+
+/**
+ * Implements the Windows target platform module.
+ */
+class FWindowsTargetPlatformControlsModule
+	: public ITargetPlatformControlsModule
+{
+public:
+
+	/** Destructor. */
+	~FWindowsTargetPlatformControlsModule( )
+	{
+
+	}
+
+public:
+	virtual void GetTargetPlatformControls(TArray<ITargetPlatformControls*>& TargetPlatforms, FName& PlatformSettingsModuleName)
+	{
+		TMap<FString, ITargetPlatformSettings*> PlatformNameToPlatformSettings;
+		IWindowsTargetPlatformSettingsModule* ModuleSettings = FModuleManager::GetModulePtr<IWindowsTargetPlatformSettingsModule>(PlatformSettingsModuleName);
+		if (ModuleSettings != nullptr)
+		{
+			TMap<FString, ITargetPlatformSettings*> OutMap;
+			ModuleSettings->GetPlatformSettingsMaps(OutMap);
+			ITargetPlatformControls* GameTP = new TGenericWindowsTargetPlatformControls<FWindowsPlatformProperties<false, false, false>>(OutMap[FWindowsPlatformProperties<false, false, false>::PlatformName()]);
+			ITargetPlatformControls* EditorTP = new TGenericWindowsTargetPlatformControls<FWindowsPlatformProperties<true, false, false>>(OutMap[FWindowsPlatformProperties<true, false, false>::PlatformName()]);
+			ITargetPlatformControls* ServerTP = new TGenericWindowsTargetPlatformControls<FWindowsPlatformProperties<false, true, false>>(OutMap[FWindowsPlatformProperties<false, true, false>::PlatformName()]);
+			ITargetPlatformControls* ClientTP = new TGenericWindowsTargetPlatformControls<FWindowsPlatformProperties<false, false, true>>(OutMap[FWindowsPlatformProperties<false, false, true>::PlatformName()]);
+			TargetPlatforms.Add(GameTP);
+			TargetPlatforms.Add(EditorTP);
+			TargetPlatforms.Add(ServerTP);
+			TargetPlatforms.Add(ClientTP);
+#if WITH_ENGINE
+			// currently this TP requires the engine for allowing GameDelegates usage
+			bool bSupportCookedEditor;
+			if (GConfig->GetBool(TEXT("CookedEditorSettings"), TEXT("bSupportCookedEditor"), bSupportCookedEditor, GGameIni) && bSupportCookedEditor)
+			{
+				ITargetPlatformControls* CookedEditorTP = new TCookedEditorTargetPlatformControls<FWindowsEditorTargetPlatformControlsParent>(ModuleSettings->GetCookedEditorPlatformSettings());
+				ITargetPlatformControls* CookedCookerTP = new TCookedCookerTargetPlatformControls<FWindowsEditorTargetPlatformControlsParent>(ModuleSettings->GetCookedCookerPlatformSettings());
+				TargetPlatforms.Add(CookedEditorTP);
+				TargetPlatforms.Add(CookedCookerTP);
+			}
+#endif
+		}
+	}
+
+public:
+};
+
+
+#undef LOCTEXT_NAMESPACE
+
+
+IMPLEMENT_MODULE(FWindowsTargetPlatformControlsModule, WindowsTargetPlatformControls);
diff --git a/Engine/Source/Developer/Windows/WindowsTargetPlatfomrControls/Public/GenericWindowsTargetPlatformControls.h b/Engine/Source/Developer/Windows/WindowsTargetPlatfomrControls/Public/GenericWindowsTargetPlatformControls.h
new file mode 100644
--- /dev/null
+++ b/Engine/Source/Developer/Windows/WindowsTargetPlatfomrControls/Public/GenericWindowsTargetPlatformControls.h
@@ -0,0 +1,188 @@
+// Copyright Epic Games, Inc. All Rights Reserved.
+
+#pragma once
+
+#include "CoreMinimal.h"
+#include "Common/TargetPlatformControlsBase.h"
+#include "AnalyticsEventAttribute.h"
+#if PLATFORM_WINDOWS
+#include "LocalPcTargetDevice.h"
+#endif
+#if WITH_ENGINE
+#include "TextureResource.h"
+#endif
+#include "SteamDeck/SteamDeckDevice.h"
+#include "GenericWindowsTargetPlatformSettings.h"
+
+#define LOCTEXT_NAMESPACE "TGenericWindowsTargetPlatformControls"
+
+/**
+ * Template for Windows target platforms controls
+ */
+#if PLATFORM_WINDOWS
+template<typename TProperties, typename TTargetDevice = FLocalPcTargetDevice>
+#else
+template<typename TProperties>
+#endif
+
+class TGenericWindowsTargetPlatformControls : public TTargetPlatformControlsBase<TProperties>
+{
+public:
+	typedef TTargetPlatformControlsBase<TProperties> TSuper;
+	typedef TTargetPlatformSettingsBase<TProperties> TSuperSettings;
+
+	/**
+	O * Default constructor.
+	 */
+	TGenericWindowsTargetPlatformControls(ITargetPlatformSettings* TargetPlatformSettings)
+		: TSuper(TargetPlatformSettings)
+	{
+#if PLATFORM_WINDOWS
+		// only add local device if actually running on Windows
+		LocalDevice = MakeShareable(new TTargetDevice(*this));
+
+		// quick solution to not having WinGDK steamdeck devices
+		if (this->PlatformName().StartsWith(TEXT("Windows")))
+		{
+			// Check if we have any SteamDeck devices around
+			SteamDevices = TSteamDeckDevice<FLocalPcTargetDevice>::DiscoverDevices(*this, TEXT("Proton"));
+		}
+#endif
+		GenericWindowsTargetPlatformSettings = (TSuperSettings*)TargetPlatformSettings;
+	}
+
+public:
+
+	//~ Begin ITargetPlatform Interface
+
+	virtual void EnableDeviceCheck(bool OnOff) override {}
+
+	virtual void GetAllDevices( TArray<ITargetDevicePtr>& OutDevices ) const override
+	{
+		OutDevices.Reset();
+		if (LocalDevice.IsValid())
+		{
+			OutDevices.Add(LocalDevice);
+		}
+
+		for (const ITargetDevicePtr& SteamDeck : SteamDevices)
+		{
+			if (SteamDeck.IsValid())
+			{
+				OutDevices.Add(SteamDeck);
+			}
+		}
+	}
+
+	virtual bool GenerateStreamingInstallManifest(const TMultiMap<FString, int32>& PakchunkMap, const TSet<int32>& PakchunkIndicesInUse) const override
+	{
+		return true;
+	}
+
+	virtual ITargetDevicePtr GetDefaultDevice( ) const override
+	{
+		if (LocalDevice.IsValid())
+		{
+			return LocalDevice;
+		}
+
+		return nullptr;
+	}
+
+	virtual ITargetDevicePtr GetDevice( const FTargetDeviceId& DeviceId )
+	{
+		if (LocalDevice.IsValid() && (DeviceId == LocalDevice->GetId()))
+		{
+			return LocalDevice;
+		}
+
+		for (const ITargetDevicePtr& SteamDeck : SteamDevices)
+		{
+			if (SteamDeck.IsValid() && DeviceId == SteamDeck->GetId())
+			{
+				return SteamDeck;
+			}
+		}
+
+		return nullptr;
+	}
+
+	virtual bool IsRunningPlatform( ) const override
+	{
+		// Must be Windows platform as editor for this to be considered a running platform
+		return PLATFORM_WINDOWS && !UE_SERVER && !UE_GAME && WITH_EDITOR && TProperties::HasEditorOnlyData();
+	}
+
+	virtual void GetShaderCompilerDependencies(TArray<FString>& OutDependencies) const override
+	{		
+		FTargetPlatformControlsBase::AddDependencySCArrayHelper(OutDependencies, TEXT("Binaries/ThirdParty/Windows/DirectX/x64/d3dcompiler_47.dll"));
+		FTargetPlatformControlsBase::AddDependencySCArrayHelper(OutDependencies, TEXT("Binaries/ThirdParty/ShaderConductor/Win64/ShaderConductor.dll"));
+		FTargetPlatformControlsBase::AddDependencySCArrayHelper(OutDependencies, TEXT("Binaries/ThirdParty/ShaderConductor/Win64/dxcompiler.dll"));
+		FTargetPlatformControlsBase::AddDependencySCArrayHelper(OutDependencies, TEXT("Binaries/ThirdParty/ShaderConductor/Win64/dxil.dll"));
+	}
+
+	virtual void GetBuildProjectSettingKeys(FString& OutSection, TArray<FString>& InBoolKeys, TArray<FString>& InIntKeys, TArray<FString>& InStringKeys) const override
+	{
+		OutSection = TEXT("/Script/WindowsTargetPlatform.WindowsTargetSettings");
+		InStringKeys.Add(TEXT("MinimumOSVersion"));
+	}
+
+
+	virtual void GetPlatformSpecificProjectAnalytics( TArray<FAnalyticsEventAttribute>& AnalyticsParamArray ) const override
+	{
+		TSuper::GetPlatformSpecificProjectAnalytics(AnalyticsParamArray);
+
+		AppendAnalyticsEventAttributeArray(AnalyticsParamArray,
+			TEXT("UsesRayTracing"), GenericWindowsTargetPlatformSettings->UsesRayTracing()
+		);
+	
+		TSuper::AppendAnalyticsEventConfigString(AnalyticsParamArray, TEXT("/Script/WindowsTargetPlatform.WindowsTargetSettings"), TEXT("DefaultGraphicsRHI"), GEngineIni);
+
+		TSuper::AppendAnalyticsEventConfigArray(AnalyticsParamArray, TEXT("/Script/WindowsTargetPlatform.WindowsTargetSettings"), TEXT("D3D12TargetedShaderFormats"), GEngineIni);
+		TSuper::AppendAnalyticsEventConfigArray(AnalyticsParamArray, TEXT("/Script/WindowsTargetPlatform.WindowsTargetSettings"), TEXT("D3D11TargetedShaderFormats"), GEngineIni);
+		TSuper::AppendAnalyticsEventConfigArray(AnalyticsParamArray, TEXT("/Script/WindowsTargetPlatform.WindowsTargetSettings"), TEXT("VulkanTargetedShaderFormats"), GEngineIni);
+		TSuper::AppendAnalyticsEventConfigArray(AnalyticsParamArray, TEXT("/Script/WindowsTargetPlatform.WindowsTargetSettings"), TEXT("TargetedRHIs"), GEngineIni, TEXT("TargetedRHIs_Deprecated") );
+	}
+
+#if WITH_ENGINE
+
+	virtual void GetTextureFormats( const UTexture* InTexture, TArray< TArray<FName> >& OutFormats) const override
+	{
+		if (!TProperties::IsServerOnly())
+		{
+			GetDefaultTextureFormatNamePerLayer(OutFormats.AddDefaulted_GetRef(), GenericWindowsTargetPlatformSettings, this, InTexture, true, 4, true);
+		}
+	}
+
+	virtual void GetAllTextureFormats(TArray<FName>& OutFormats) const override
+	{
+		if (!TProperties::IsServerOnly())
+		{
+			GetAllDefaultTextureFormats(GenericWindowsTargetPlatformSettings, OutFormats);
+		}
+	}
+
+#endif //WITH_ENGINE
+
+	virtual bool SupportsVariants() const override
+	{
+		return true;
+	}
+
+	virtual float GetVariantPriority() const override
+	{
+		return TProperties::GetVariantPriority();
+	}
+	//~ End ITargetPlatform Interface
+
+private:
+
+	// Holds the local device.
+	ITargetDevicePtr LocalDevice;
+
+	TArray<ITargetDevicePtr> SteamDevices;
+
+	TSuperSettings* GenericWindowsTargetPlatformSettings;
+};
+
+#undef LOCTEXT_NAMESPACE
diff --git a/Engine/Source/Developer/Windows/WindowsTargetPlatfomrControls/Public/LocalPcTargetDevice.h b/Engine/Source/Developer/Windows/WindowsTargetPlatfomrControls/Public/LocalPcTargetDevice.h
new file mode 100644
--- /dev/null
+++ b/Engine/Source/Developer/Windows/WindowsTargetPlatfomrControls/Public/LocalPcTargetDevice.h
@@ -0,0 +1,330 @@
+// Copyright Epic Games, Inc. All Rights Reserved.
+
+#pragma once
+
+#include "CoreMinimal.h"
+#include "Misc/Guid.h"
+#include "HAL/FileManager.h"
+#include "HAL/PlatformProcess.h"
+#include "Interfaces/ITargetDevice.h"
+#include "Interfaces/ITargetPlatformControls.h"
+
+#include "Windows/AllowWindowsPlatformTypes.h"
+#include <TlHelp32.h>
+
+/**
+ * Template for local PC target devices.
+ *
+ */
+class FLocalPcTargetDevice
+	: public ITargetDevice
+{
+public:
+
+	/**
+	 * Creates and initializes a new device for the specified target platform.
+	 *
+	 * @param InTargetPlatform - The target platform.
+	 */
+	FLocalPcTargetDevice( const ITargetPlatformControls& InTargetPlatform )
+		: TargetPlatformControls(InTargetPlatform)
+	{ }
+
+public:
+
+	// ITargetDevice interface
+
+	virtual bool Connect( ) override
+	{
+		return true;
+	}
+
+	virtual void Disconnect( )
+	{ }
+
+	virtual ETargetDeviceTypes GetDeviceType( ) const override
+	{
+		if (::GetSystemMetrics(SM_TABLETPC) != 0)
+		{
+			return ETargetDeviceTypes::Tablet;
+		}
+
+		return ETargetDeviceTypes::Desktop;
+	}
+
+	virtual FTargetDeviceId GetId( ) const override
+	{
+		return FTargetDeviceId(TargetPlatformControls.PlatformName(), GetName());
+	}
+
+	virtual FString GetName( ) const override
+	{
+		return FPlatformProcess::ComputerName();
+	}
+
+	virtual FString GetOperatingSystemName( ) override
+	{
+		return TEXT("Windows (64-bit)");
+	}
+
+	virtual int32 GetProcessSnapshot( TArray<FTargetDeviceProcessInfo>& OutProcessInfos ) override
+	{
+		// enumerate processes
+		HANDLE ProcessSnapshot = ::CreateToolhelp32Snapshot(TH32CS_SNAPPROCESS, 0);
+
+		if (ProcessSnapshot != INVALID_HANDLE_VALUE)
+		{
+			PROCESSENTRY32 ProcessEntry;
+			ProcessEntry.dwSize = sizeof(PROCESSENTRY32);
+
+			if (::Process32First(ProcessSnapshot, &ProcessEntry))
+			{
+				do 
+				{
+					// only include processes that the user has permission to query
+					HANDLE ProcessHandle = ::OpenProcess(PROCESS_QUERY_INFORMATION, false, ProcessEntry.th32ProcessID);
+
+					if (ProcessHandle != NULL)
+					{
+						FTargetDeviceProcessInfo ProcessInfo;
+
+						// get process details
+						ProcessInfo.Id = ProcessEntry.th32ProcessID;
+						ProcessInfo.Name = ProcessEntry.szExeFile;
+						ProcessInfo.ParentId = ProcessEntry.th32ParentProcessID;
+						ProcessInfo.UserName = TEXT("-");
+
+						// get process user name
+						HANDLE TokenHandle;
+
+						if (::OpenProcessToken(ProcessHandle, TOKEN_QUERY, &TokenHandle))
+						{
+							DWORD UserTokenSize;
+
+							::GetTokenInformation(TokenHandle, TokenUser, NULL, 0, &UserTokenSize);
+
+							if (::GetLastError() == ERROR_INSUFFICIENT_BUFFER)
+							{
+								PTOKEN_USER UserToken = reinterpret_cast<PTOKEN_USER>(new BYTE[UserTokenSize]);
+
+								if (UserToken != NULL)
+								{
+									if (::GetTokenInformation(TokenHandle, TokenUser, UserToken, UserTokenSize, &UserTokenSize))
+									{
+										WCHAR DomainName[256];
+										DWORD DomainNameLength = 256;
+
+										WCHAR UserName[256];
+										DWORD UserNameLength = 256;
+
+										SID_NAME_USE SidType;
+
+										if (::LookupAccountSid(NULL, UserToken->User.Sid, UserName, &UserNameLength, DomainName, &DomainNameLength, &SidType))
+										{
+											ProcessInfo.UserName = UserName;
+										}
+										else
+										{
+											ProcessInfo.UserName = TEXT("SYSTEM");
+										}
+									}
+								}
+								else
+								{
+									DWORD LastError = ::GetLastError();
+								}
+							}
+
+							::CloseHandle(TokenHandle);
+						}
+
+						::CloseHandle(ProcessHandle);
+
+						OutProcessInfos.Add(ProcessInfo);
+					}
+				}
+				while (::Process32Next(ProcessSnapshot, &ProcessEntry));
+			}
+
+			::CloseHandle(ProcessSnapshot);
+		}
+
+		// get thread details
+		HANDLE ThreadSnapshot = ::CreateToolhelp32Snapshot(TH32CS_SNAPTHREAD, 0);
+
+		if (ThreadSnapshot != INVALID_HANDLE_VALUE)
+		{
+			THREADENTRY32 ThreadEntry;
+			ThreadEntry.dwSize = sizeof(THREADENTRY32);
+
+			if (::Thread32First(ThreadSnapshot, &ThreadEntry))
+			{
+				do 
+				{
+					for (int32 ProcessInfoIndex = 0; ProcessInfoIndex < OutProcessInfos.Num(); ++ProcessInfoIndex)
+					{
+						FTargetDeviceProcessInfo& ProcessInfo = OutProcessInfos[ProcessInfoIndex];
+
+						if (ProcessInfo.Id == ThreadEntry.th32OwnerProcessID)
+						{
+							FTargetDeviceThreadInfo ThreadInfo;
+
+							ThreadInfo.ExitCode = 0;
+							ThreadInfo.Id = ThreadEntry.th32ThreadID;
+							ThreadInfo.StackSize = 0;
+							ThreadInfo.State = ETargetDeviceThreadStates::Unknown;
+							ThreadInfo.WaitState = ETargetDeviceThreadWaitStates::Unknown;
+
+							ProcessInfo.Threads.Add(ThreadInfo);
+
+							break;
+						}
+					}
+				}
+				while (::Thread32Next(ThreadSnapshot, &ThreadEntry));
+
+				::CloseHandle(ThreadSnapshot);
+			}
+		}
+
+		return OutProcessInfos.Num();
+	}
+
+	virtual const class ITargetPlatformSettings& GetPlatformSettings() const override
+	{
+		return *(TargetPlatformControls.GetTargetPlatformSettings());
+	}
+	virtual const class ITargetPlatformControls& GetPlatformControls() const override
+	{
+		return TargetPlatformControls;
+	}
+
+	virtual bool IsConnected( )
+	{
+		return true;
+	}
+
+	virtual bool IsDefault( ) const override
+	{
+		return true;
+	}
+
+	virtual bool PowerOff( bool Force ) override
+	{
+		if (!AdjustShutdownPrivileges())
+		{
+			return false;
+		}
+
+		CA_SUPPRESS(28159);
+		return (::ExitWindowsEx(EWX_POWEROFF | EWX_FORCE, SHTDN_REASON_MINOR_MAINTENANCE | SHTDN_REASON_FLAG_PLANNED) != 0);
+	}
+
+	virtual bool PowerOn( ) override
+	{
+		return false;
+	}
+
+	virtual bool Reboot( bool bReconnect = false ) override
+	{
+		if (!AdjustShutdownPrivileges())
+		{
+			return false;
+		}
+
+		CA_SUPPRESS(28159);
+		return (::ExitWindowsEx(EWX_REBOOT | EWX_FORCE, SHTDN_REASON_MINOR_MAINTENANCE | SHTDN_REASON_FLAG_PLANNED) != 0);
+	}
+
+	virtual bool SupportsFeature( ETargetDeviceFeatures Feature ) const override
+	{
+		switch (Feature)
+		{
+		case ETargetDeviceFeatures::MultiLaunch:
+			return true;
+
+		case ETargetDeviceFeatures::PowerOff:
+			return true;
+
+		// @todo gmp: implement turning on remote PCs (wake on LAN)
+		case ETargetDeviceFeatures::PowerOn:
+			return false;
+
+		case ETargetDeviceFeatures::ProcessSnapshot:
+			return true;
+
+		case ETargetDeviceFeatures::Reboot:
+			return true;
+		}
+
+		return false;
+	}
+
+	virtual void SetUserCredentials( const FString& UserName, const FString& UserPassword ) override
+	{
+	}
+
+	virtual bool GetUserCredentials( FString& OutUserName, FString& OutUserPassword ) override
+	{
+		return false;
+	}
+
+	virtual bool TerminateProcess( const int64 ProcessId ) override
+	{
+		HANDLE ProcessHandle = OpenProcess(PROCESS_TERMINATE, false, (DWORD)ProcessId );
+
+		if (ProcessHandle != NULL)
+		{
+			bool Result = (::TerminateProcess(ProcessHandle, 0) != 0);
+
+			CloseHandle(ProcessHandle);
+
+			return Result;
+		}
+
+		return false;
+	}
+
+protected:
+
+	/**
+	 * Adjusts shutdown privileges for the local host PC.
+	 *
+	 * @return true on success, false otherwise.
+	 */
+	bool AdjustShutdownPrivileges( )
+	{
+		HANDLE TokenHandle;
+		TOKEN_PRIVILEGES Token;
+
+		if (!::OpenProcessToken(GetCurrentProcess(), TOKEN_ADJUST_PRIVILEGES | TOKEN_QUERY, &TokenHandle))
+		{
+			return false;
+		}
+
+		::LookupPrivilegeValue(NULL, SE_SHUTDOWN_NAME, &Token.Privileges[0].Luid);
+
+		if (Token.Privileges[0].Attributes == SE_PRIVILEGE_ENABLED)
+		{
+			return true;
+		}
+
+		Token.PrivilegeCount = 1;
+		Token.Privileges[0].Attributes = SE_PRIVILEGE_ENABLED;
+
+		::AdjustTokenPrivileges(TokenHandle, 0, &Token, 0, (PTOKEN_PRIVILEGES)NULL, 0);
+      
+		return (GetLastError() == ERROR_SUCCESS);
+	}
+
+protected:
+
+	// Holds the collection of processes that were started using the Run() method.
+	TMap<FGuid, void*> Processes;
+
+	// Holds a reference to the device's target platform.
+	const ITargetPlatformControls& TargetPlatformControls;
+};
+
+#include "Windows/HideWindowsPlatformTypes.h"
+#undef PROCESSENTRY32
diff --git a/Engine/Source/Developer/Windows/WindowsTargetPlatfomrControls/WindowsTargetPlatformControls.Build.cs b/Engine/Source/Developer/Windows/WindowsTargetPlatfomrControls/WindowsTargetPlatformControls.Build.cs
new file mode 100644
--- /dev/null
+++ b/Engine/Source/Developer/Windows/WindowsTargetPlatfomrControls/WindowsTargetPlatformControls.Build.cs
@@ -0,0 +1,50 @@
+// Copyright Epic Games, Inc. All Rights Reserved.
+
+using UnrealBuildTool;
+
+public class WindowsTargetPlatformControls : ModuleRules
+{
+	public WindowsTargetPlatformControls(ReadOnlyTargetRules Target) : base(Target)
+	{
+		// We need a short name here since this can run afoul very easily of the `MAX_PATH` limit when
+		// combined with building other targets that make use of this.
+		ShortName = "WinTPCon";
+
+		PrivateDependencyModuleNames.AddRange(
+			new string[] {
+				"Core",
+				"CoreUObject",
+				"TargetPlatform",
+				"DesktopPlatform",
+                "AudioPlatformConfiguration",
+				"WindowsTargetPlatformSettings",
+			}
+		);
+
+		PrivateIncludePathModuleNames.AddRange(
+			new string[] {
+				"Settings",
+				"WindowsTargetPlatformSettings"
+			}
+		);
+
+		PublicIncludePathModuleNames.AddRange(
+			new string[] {
+				"AudioPlatformConfiguration"
+			}
+		);
+
+		// compile with Engine
+		if (Target.bCompileAgainstEngine)
+		{
+			PublicIncludePathModuleNames.Add("Engine");
+			PrivateDependencyModuleNames.AddRange( new string[] {
+				"Engine", 
+				"RHI",
+				"CookedEditor",
+				}
+			);
+            PrivateIncludePathModuleNames.Add("TextureCompressor");
+        }
+    }
+}
diff --git a/Engine/Source/Developer/Windows/WindowsTargetPlatform/Classes/WindowsTargetSettings.h b/Engine/Source/Developer/Windows/WindowsTargetPlatform/Classes/WindowsTargetSettings.h
deleted file mode 100644
diff --git a/Engine/Source/Developer/Windows/WindowsTargetPlatform/Private/WindowsTargetPlatformClasses.cpp b/Engine/Source/Developer/Windows/WindowsTargetPlatform/Private/WindowsTargetPlatformClasses.cpp
deleted file mode 100644
diff --git a/Engine/Source/Developer/Windows/WindowsTargetPlatform/Private/WindowsTargetPlatformModule.cpp b/Engine/Source/Developer/Windows/WindowsTargetPlatform/Private/WindowsTargetPlatformModule.cpp
--- a/Engine/Source/Developer/Windows/WindowsTargetPlatform/Private/WindowsTargetPlatformModule.cpp
+++ b/Engine/Source/Developer/Windows/WindowsTargetPlatform/Private/WindowsTargetPlatformModule.cpp
@@ -1,15 +1,9 @@
 // Copyright Epic Games, Inc. All Rights Reserved.
 
 #include "CoreMinimal.h"
-#include "ISettingsModule.h"
 #include "Interfaces/ITargetPlatformModule.h"
-#include "GenericWindowsTargetPlatform.h"
 #include "Modules/ModuleManager.h"
-#include "UObject/Package.h"
-#include "UObject/WeakObjectPtr.h"
-#if WITH_ENGINE
-#include "CookedEditorTargetPlatform.h"
-#endif
+#include "Common/TargetPlatformBase.h"
 
 #define LOCTEXT_NAMESPACE "FWindowsTargetPlatformModule"
 
@@ -30,68 +24,19 @@ public:
 
 public:
 
-	// this is an example of a hotfix, declared here for no particular reason. Once we have other examples, it can be deleted.
-#if 0
-	void HotfixTest( void *InPayload, int PayloadSize )
+	virtual void GetTargetPlatforms(TArray<ITargetPlatform*>& TargetPlatforms) override
 	{
-		check(sizeof(FTestHotFixPayload) == PayloadSize);
-		
-		FTestHotFixPayload* Payload = (FTestHotFixPayload*)InPayload;
-		UE_LOG(LogTemp, Log, TEXT("Hotfix Test %s"), *Payload->Message);
-		Payload->Result = Payload->ValueToReturn;
 	}
-#endif
-
-public:
 
-	virtual void GetTargetPlatforms(TArray<ITargetPlatform*>& TargetPlatforms) override
+	virtual void GetTargetPlatforms(TArray<ITargetPlatform*>& TargetPlatforms, TArray<ITargetPlatformSettings*> TargetPlatformSettings, TArray<ITargetPlatformControls*> TargetPlatformControls) override
 	{
-		// Game TP
-		TargetPlatforms.Add(new TGenericWindowsTargetPlatform<FWindowsPlatformProperties<false, false, false>>());
-		// Editor TP
-		TargetPlatforms.Add(new TGenericWindowsTargetPlatform<FWindowsPlatformProperties<true, false, false>>());
-		// Server TP
-		TargetPlatforms.Add(new TGenericWindowsTargetPlatform<FWindowsPlatformProperties<false, true, false>>());
-		// Client TP
-		TargetPlatforms.Add(new TGenericWindowsTargetPlatform<FWindowsPlatformProperties<false, false, true>>());
-
-#if WITH_ENGINE
-		// currently this TP requires the engine for allowing GameDelegates usage
-		bool bSupportCookedEditor;
-		if (GConfig->GetBool(TEXT("CookedEditorSettings"), TEXT("bSupportCookedEditor"), bSupportCookedEditor, GGameIni) && bSupportCookedEditor)
+		for (ITargetPlatformControls* TargetPlatformControlsIt : TargetPlatformControls)
 		{
-			TargetPlatforms.Add(new TCookedEditorTargetPlatform<FWindowsEditorTargetPlatformParent>());
-			TargetPlatforms.Add(new TCookedCookerTargetPlatform<FWindowsEditorTargetPlatformParent>());
+			TargetPlatforms.Add(new FTargetPlatformMerged(TargetPlatformControlsIt->GetTargetPlatformSettings(), TargetPlatformControlsIt));
 		}
-#endif
-	}
-
-public:
-
-	// IModuleInterface interface
-
-	virtual void StartupModule() override
-	{
-		// this is an example of a hotfix, declared here for no particular reason. Once we have other examples, it can be deleted.
-#if 0
-		FCoreDelegates::GetHotfixDelegate(EHotfixDelegates::Test).BindRaw(this, &FWindowsTargetPlatformModule::HotfixTest);
-#endif
 	}
-
-	virtual void ShutdownModule() override
-	{
-		// this is an example of a hotfix, declared here for no particular reason. Once we have other examples, it can be deleted.
-#if 0
-		FCoreDelegates::GetHotfixDelegate(EHotfixDelegates::Test).Unbind();
-#endif
-
-		ISettingsModule* SettingsModule = FModuleManager::GetModulePtr<ISettingsModule>("Settings");
-	}
-
 };
 
-
 #undef LOCTEXT_NAMESPACE
 
-
 IMPLEMENT_MODULE(FWindowsTargetPlatformModule, WindowsTargetPlatform);
diff --git a/Engine/Source/Developer/Windows/WindowsTargetPlatform/Public/GenericWindowsTargetPlatform.h b/Engine/Source/Developer/Windows/WindowsTargetPlatform/Public/GenericWindowsTargetPlatform.h
deleted file mode 100644
diff --git a/Engine/Source/Developer/Windows/WindowsTargetPlatform/Public/LocalPcTargetDevice.h b/Engine/Source/Developer/Windows/WindowsTargetPlatform/Public/LocalPcTargetDevice.h
deleted file mode 100644
diff --git a/Engine/Source/Developer/Windows/WindowsTargetPlatform/WindowsTargetPlatform.Build.cs b/Engine/Source/Developer/Windows/WindowsTargetPlatform/WindowsTargetPlatform.Build.cs
--- a/Engine/Source/Developer/Windows/WindowsTargetPlatform/WindowsTargetPlatform.Build.cs
+++ b/Engine/Source/Developer/Windows/WindowsTargetPlatform/WindowsTargetPlatform.Build.cs
@@ -9,36 +9,11 @@ public class WindowsTargetPlatform : ModuleRules
         PrivateDependencyModuleNames.AddRange(
 			new string[] {
 				"Core",
-				"CoreUObject",
 				"TargetPlatform",
 				"DesktopPlatform",
-                "AudioPlatformConfiguration",
-            }
-		);
-
-		PrivateIncludePathModuleNames.AddRange(
-			new string[] {
-				"Settings"
-			}
-		);
-
-		PublicIncludePathModuleNames.AddRange(
-			new string[] {
-				"AudioPlatformConfiguration"
+				"WindowsTargetPlatformSettings",
+				"WindowsTargetPlatformControls",
 			}
 		);
-
-		// compile with Engine
-		if (Target.bCompileAgainstEngine)
-		{
-			PublicIncludePathModuleNames.Add("Engine");
-			PrivateDependencyModuleNames.AddRange( new string[] {
-				"Engine", 
-				"RHI",
-				"CookedEditor",
-				}
-			);
-            PrivateIncludePathModuleNames.Add("TextureCompressor");
-        }
     }
 }
diff --git a/Engine/Source/Developer/Windows/WindowsTargetPlatformSettings/Classes/WindowsTargetSettings.h b/Engine/Source/Developer/Windows/WindowsTargetPlatformSettings/Classes/WindowsTargetSettings.h
new file mode 100644
--- /dev/null
+++ b/Engine/Source/Developer/Windows/WindowsTargetPlatformSettings/Classes/WindowsTargetSettings.h
@@ -0,0 +1,148 @@
+// Copyright Epic Games, Inc. All Rights Reserved.
+
+/*=============================================================================
+	WindowsTargetSettings.h: Declares the UWindowsTargetSettings class.
+=============================================================================*/
+
+#pragma once
+
+#include "CoreMinimal.h"
+#include "UObject/Object.h"
+#include "AudioCompressionSettings.h"
+#include "WindowsTargetSettings.generated.h"
+
+UENUM()
+enum class ECompilerVersion : uint8
+{
+	Default = 0,
+	VisualStudio2015 = 1 UMETA(DisplayName = "Visual Studio 2015 (deprecated)"),
+	VisualStudio2017 = 2 UMETA(DisplayName = "Visual Studio 2017 (deprecated)"),
+	VisualStudio2019 = 3 UMETA(DisplayName = "Visual Studio 2019 (deprecated)"),
+	VisualStudio2022 = 4 UMETA(DisplayName = "Visual Studio 2022"),
+};
+
+UENUM()
+enum class EDefaultGraphicsRHI : uint8
+{
+	DefaultGraphicsRHI_Default = 0 UMETA(DisplayName = "Default"),
+	DefaultGraphicsRHI_DX11 = 1 UMETA(DisplayName = "DirectX 11"),
+	DefaultGraphicsRHI_DX12 = 2 UMETA(DisplayName = "DirectX 12"),
+	DefaultGraphicsRHI_Vulkan = 3 UMETA(DisplayName = "Vulkan"),
+};
+
+/**
+ * Implements the settings for the Windows target platform. The first instance of this class is initialized in
+ * WindowsTargetPlatform, really early during the startup sequence before the CDO has been constructed, so its config 
+ * settings are read manually from there.
+ */
+UCLASS(config=Engine, defaultconfig)
+class WINDOWSTARGETPLATFORMSETTINGS_API UWindowsTargetSettings
+	: public UObject
+{
+public:
+	GENERATED_UCLASS_BODY()
+
+	virtual void PostInitProperties() override;
+
+	/** Select which RHI to use. Make sure its also selected as a Targeted RHI. Requires Editor restart. */
+	UPROPERTY(EditAnywhere, config, Category="Targeted RHIs", Meta = (DisplayName = "Default RHI", ConfigRestartRequired = true))
+	EDefaultGraphicsRHI DefaultGraphicsRHI;
+
+	UPROPERTY(config, meta = (DeprecatedProperty, DeprecationMessage = "Use one of the RHI specific lists."))
+	TArray<FString> TargetedRHIs_DEPRECATED;
+
+	UPROPERTY(EditAnywhere, config, Category = "Rendering", Meta = (ConfigRestartRequired = true))
+	TArray<FString> D3D12TargetedShaderFormats;
+
+	UPROPERTY(EditAnywhere, config, Category = "Rendering", Meta = (ConfigRestartRequired = true))
+	TArray<FString> D3D11TargetedShaderFormats;
+
+	UPROPERTY(EditAnywhere, config, Category = "Rendering", Meta = (ConfigRestartRequired = true))
+	TArray<FString> VulkanTargetedShaderFormats;
+
+	/** The compiler version to use for this project. May be different to the chosen IDE. */
+	UPROPERTY(EditAnywhere, config, Category = "Toolchain", Meta = (DisplayName = "Compiler Version"))
+	ECompilerVersion Compiler;
+
+	/** Sample rate to run the audio mixer with. */
+	UPROPERTY(config, EditAnywhere, Category = "Audio", Meta = (DisplayName = "Audio Mixer Sample Rate"))
+	int32 AudioSampleRate;
+
+	/** The amount of audio to compute each callback block. Lower values decrease latency but may increase CPU cost. */
+	UPROPERTY(config, EditAnywhere, Category = "Audio", meta = (ClampMin = "64", ClampMax = "4096", DisplayName = "Callback Buffer Size"))
+	int32 AudioCallbackBufferFrameSize;
+
+	/** The number of buffers to keep enqueued. More buffers increases latency, but can compensate for variable compute availability in audio callbacks on some platforms. */
+	UPROPERTY(config, EditAnywhere, Category = "Audio", meta = (ClampMin = "1", UIMin = "1", DisplayName = "Number of Buffers To Enqueue"))
+	int32 AudioNumBuffersToEnqueue;
+
+	/** The max number of channels (voices) to limit for this platform. The max channels used will be the minimum of this value and the global audio quality settings. A value of 0 will not apply a platform channel count max. */
+	UPROPERTY(config, EditAnywhere, Category = "Audio", meta = (ClampMin = "0", UIMin = "0", DisplayName = "Max Channels"))
+	int32 AudioMaxChannels;
+
+	/** The number of workers to use to compute source audio. Will only use up to the max number of sources. Will evenly divide sources to each source worker. */
+	UPROPERTY(config, EditAnywhere, Category = "Audio", meta = (ClampMin = "0", UIMin = "0", DisplayName = "Number of Source Workers"))
+	int32 AudioNumSourceWorkers;
+
+	/** Which of the currently enabled spatialization plugins to use. */
+	UPROPERTY(config, EditAnywhere, Category = "Audio")
+	FString SpatializationPlugin;
+
+	/** Which of the currently enabled source data override plugins to use. */
+	UPROPERTY(config, EditAnywhere, Category = "Audio")
+	FString SourceDataOverridePlugin;
+
+	/** Which of the currently enabled reverb plugins to use. */
+	UPROPERTY(config, EditAnywhere, Category = "Audio")
+	FString ReverbPlugin;
+
+	/** Which of the currently enabled occlusion plugins to use. */
+	UPROPERTY(config, EditAnywhere, Category = "Audio")
+	FString OcclusionPlugin;
+
+	/** Audio Cook Settings */
+
+	/** Various overrides for how this platform should handle compression and decompression */
+	UPROPERTY(GlobalConfig, EditAnywhere, Category = "Audio")
+	FPlatformRuntimeAudioCompressionOverrides CompressionOverrides;
+
+	/** This determines the max amount of memory that should be used for the cache at any given time. If set low (<= 8 MB), it lowers the size of individual chunks of audio during cook. */
+	UPROPERTY(GlobalConfig, EditAnywhere, Category = "Audio|CookOverrides|Stream Caching", meta = (DisplayName = "Max Cache Size (KB)"))
+	int32 CacheSizeKB;
+
+	/** This overrides the default max chunk size used when chunking audio for stream caching (ignored if < 0) */
+	UPROPERTY(GlobalConfig, EditAnywhere, Category = "Audio|CookOverrides|Stream Caching", meta = (DisplayName = "Max Chunk Size Override (KB)"))
+	int32 MaxChunkSizeOverrideKB;
+
+	UPROPERTY(GlobalConfig, EditAnywhere, Category = "Audio|CookOverrides")
+	bool bResampleForDevice;
+
+	/** Mapping of which sample rates are used for each sample rate quality for a specific platform. */
+
+	UPROPERTY(GlobalConfig, EditAnywhere, Category = "Audio|CookOverrides|ResamplingQuality", meta = (DisplayName = "Max"))
+	float MaxSampleRate;
+
+	UPROPERTY(GlobalConfig, EditAnywhere, Category = "Audio|CookOverrides|ResamplingQuality", meta = (DisplayName = "High"))
+	float HighSampleRate;
+
+	UPROPERTY(GlobalConfig, EditAnywhere, Category = "Audio|CookOverrides|ResamplingQuality", meta = (DisplayName = "Medium"))
+	float MedSampleRate;
+
+	UPROPERTY(GlobalConfig, EditAnywhere, Category = "Audio|CookOverrides|ResamplingQuality", meta = (DisplayName = "Low"))
+	float LowSampleRate;
+
+	UPROPERTY(GlobalConfig, EditAnywhere, Category = "Audio|CookOverrides|ResamplingQuality", meta = (DisplayName = "Min"))
+	float MinSampleRate;
+
+	/** Scales all compression qualities when cooking to this platform. For example, 0.5 will halve all compression qualities, and 1.0 will leave them unchanged. */
+	UPROPERTY(GlobalConfig, EditAnywhere, Category = "Audio|CookOverrides")
+	float CompressionQualityModifier;
+
+	/** When set to anything beyond 0, this will ensure any SoundWaves longer than this value, in seconds, to stream directly off of the disk. */
+	UPROPERTY(GlobalConfig)
+	float AutoStreamingThreshold;
+
+	/** Quality Level to COOK SoundCues at (if set, all other levels will be stripped by the cooker). */
+	UPROPERTY(GlobalConfig, EditAnywhere, Category = "Audio|CookOverrides", meta = (DisplayName = "Sound Cue Cook Quality"))
+	int32 SoundCueCookQualityIndex = INDEX_NONE;
+};
diff --git a/Engine/Source/Developer/Windows/WindowsTargetPlatformSettings/Private/WindowsTargetPlatformClasses.cpp b/Engine/Source/Developer/Windows/WindowsTargetPlatformSettings/Private/WindowsTargetPlatformClasses.cpp
new file mode 100644
--- /dev/null
+++ b/Engine/Source/Developer/Windows/WindowsTargetPlatformSettings/Private/WindowsTargetPlatformClasses.cpp
@@ -0,0 +1,64 @@
+// Copyright Epic Games, Inc. All Rights Reserved.
+
+#include "WindowsTargetSettings.h"
+
+/* UWindowsTargetSettings structors
+ *****************************************************************************/
+
+UWindowsTargetSettings::UWindowsTargetSettings( const FObjectInitializer& ObjectInitializer )
+	: Super(ObjectInitializer)
+	, CacheSizeKB(65536)
+	, MaxSampleRate(48000)
+	, HighSampleRate(32000)
+	, MedSampleRate(24000)
+	, LowSampleRate(12000)
+	, MinSampleRate(8000)
+	, CompressionQualityModifier(1)
+{
+	// Default windows settings
+	AudioSampleRate = 48000;
+	AudioCallbackBufferFrameSize = 1024;
+	AudioNumBuffersToEnqueue = 1;
+	AudioNumSourceWorkers = 4;
+}
+
+static bool FilterShaderPlatform_D3D12(const FString& InShaderPlatform)
+{
+	return InShaderPlatform == TEXT("PCD3D_SM6") || InShaderPlatform == TEXT("PCD3D_SM5") || InShaderPlatform == TEXT("PCD3D_ES31");
+}
+
+static bool FilterShaderPlatform_D3D11(const FString& InShaderPlatform)
+{
+	return InShaderPlatform == TEXT("PCD3D_SM5") || InShaderPlatform == TEXT("PCD3D_ES31");
+}
+
+static bool FilterShaderPlatform_Vulkan(const FString& InShaderPlatform)
+{
+	return InShaderPlatform == TEXT("SF_VULKAN_SM5") || InShaderPlatform == TEXT("SF_VULKAN_SM6");
+}
+
+template<typename TFilterFunc>
+static void AddToShaderFormatList(TArray<FString>& Dest, const TArray<FString>& Source, TFilterFunc FilterFunc)
+{
+	for (const FString& ShaderFormat : Source)
+	{
+		if (FilterFunc(ShaderFormat))
+		{
+			Dest.AddUnique(ShaderFormat);
+		}
+	}
+}
+
+void UWindowsTargetSettings::PostInitProperties()
+{
+	Super::PostInitProperties();
+
+	if (!TargetedRHIs_DEPRECATED.IsEmpty())
+	{
+		AddToShaderFormatList(D3D12TargetedShaderFormats, TargetedRHIs_DEPRECATED, &FilterShaderPlatform_D3D12);
+		AddToShaderFormatList(D3D11TargetedShaderFormats, TargetedRHIs_DEPRECATED, &FilterShaderPlatform_D3D11);
+		AddToShaderFormatList(VulkanTargetedShaderFormats, TargetedRHIs_DEPRECATED, &FilterShaderPlatform_Vulkan);
+
+		TargetedRHIs_DEPRECATED.Empty();
+	}
+}
diff --git a/Engine/Source/Developer/Windows/WindowsTargetPlatformSettings/Private/WindowsTargetPlatformSettingsModule.cpp b/Engine/Source/Developer/Windows/WindowsTargetPlatformSettings/Private/WindowsTargetPlatformSettingsModule.cpp
new file mode 100644
--- /dev/null
+++ b/Engine/Source/Developer/Windows/WindowsTargetPlatformSettings/Private/WindowsTargetPlatformSettingsModule.cpp
@@ -0,0 +1,129 @@
+// Copyright Epic Games, Inc. All Rights Reserved.
+
+#include "CoreMinimal.h"
+#include "ISettingsModule.h"
+#include "Interfaces/ITargetPlatformSettingsModule.h"
+#include "GenericWindowsTargetPlatformSettings.h"
+#include "IWindowsTargetPlatformSettingsModule.h"
+#include "Modules/ModuleManager.h"
+#include "UObject/Package.h"
+#include "UObject/WeakObjectPtr.h"
+#if WITH_ENGINE
+#include "CookedEditorTargetPlatformSettings.h"
+#endif
+
+#define LOCTEXT_NAMESPACE "FWindowsTargetPlatformSettingsModule"
+
+
+/**
+ * Implements the Windows target platform module.
+ */
+class FWindowsTargetPlatformSettingsModule
+	: public IWindowsTargetPlatformSettingsModule
+{
+public:
+
+	/** Destructor. */
+	~FWindowsTargetPlatformSettingsModule( )
+	{
+
+	}
+
+	virtual void GetPlatformSettingsMaps(TMap<FString, ITargetPlatformSettings*>& OutMap) override
+	{
+		OutMap = PlatformNameToPlatformSettings;
+	}
+
+	virtual ITargetPlatformSettings* GetCookedEditorPlatformSettings() override
+	{
+		return PlatformSettingsCookedEditor;
+	}
+
+	virtual ITargetPlatformSettings* GetCookedCookerPlatformSettings() override
+	{
+		return PlatformSettingsCookedCooker;
+	}
+
+public:
+
+	// this is an example of a hotfix, declared here for no particular reason. Once we have other examples, it can be deleted.
+#if 0
+	void HotfixTest( void *InPayload, int PayloadSize )
+	{
+		check(sizeof(FTestHotFixPayload) == PayloadSize);
+		
+		FTestHotFixPayload* Payload = (FTestHotFixPayload*)InPayload;
+		UE_LOG(LogTemp, Log, TEXT("Hotfix Test %s"), *Payload->Message);
+		Payload->Result = Payload->ValueToReturn;
+	}
+#endif
+
+public:
+
+	virtual void GetTargetPlatformSettings(TArray<ITargetPlatformSettings*>& TargetPlatforms)
+	{
+		// Game TP
+		ITargetPlatformSettings* GameTP = new TGenericWindowsTargetPlatformSettings<FWindowsPlatformProperties<false, false, false>>();
+		TargetPlatforms.Add(GameTP);
+		PlatformNameToPlatformSettings.Add(FWindowsPlatformProperties<false, false, false>::PlatformName(), GameTP);
+		// Editor TP
+		ITargetPlatformSettings* EditorTP = new TGenericWindowsTargetPlatformSettings<FWindowsPlatformProperties<true, false, false>>();
+		TargetPlatforms.Add(new TGenericWindowsTargetPlatformSettings<FWindowsPlatformProperties<true, false, false>>());
+		PlatformNameToPlatformSettings.Add(FWindowsPlatformProperties<true, false, false>::PlatformName(), EditorTP);
+		// Server TP
+		ITargetPlatformSettings* ServerTP = new TGenericWindowsTargetPlatformSettings<FWindowsPlatformProperties<false, true, false>>();
+		TargetPlatforms.Add(new TGenericWindowsTargetPlatformSettings<FWindowsPlatformProperties<false, true, false>>());
+		PlatformNameToPlatformSettings.Add(FWindowsPlatformProperties<false, true, false>::PlatformName(), ServerTP);
+		// Client TP
+		ITargetPlatformSettings* ClientTP = new TGenericWindowsTargetPlatformSettings<FWindowsPlatformProperties<false, false, true>>();
+		TargetPlatforms.Add(new TGenericWindowsTargetPlatformSettings<FWindowsPlatformProperties<false, false, true>>());
+		PlatformNameToPlatformSettings.Add(FWindowsPlatformProperties<false, false, true>::PlatformName(), ClientTP);
+
+#if WITH_ENGINE
+		// currently this TP requires the engine for allowing GameDelegates usage
+		bool bSupportCookedEditor;
+		if (GConfig->GetBool(TEXT("CookedEditorSettings"), TEXT("bSupportCookedEditor"), bSupportCookedEditor, GGameIni) && bSupportCookedEditor)
+		{
+			ITargetPlatformSettings* CookedEditorTP = new TCookedEditorTargetPlatformSettings<FWindowsEditorTargetPlatformSettingsParent>();
+			ITargetPlatformSettings* CookedCookerTP = new TCookedCookerTargetPlatformSettings<FWindowsEditorTargetPlatformSettingsParent>();
+			PlatformSettingsCookedEditor = CookedEditorTP;
+			PlatformSettingsCookedCooker = CookedCookerTP;
+			TargetPlatforms.Add(CookedEditorTP);
+			TargetPlatforms.Add(CookedCookerTP);
+		}
+#endif
+	}
+
+public:
+
+	// IModuleInterface interface
+
+	virtual void StartupModule() override
+	{
+		// this is an example of a hotfix, declared here for no particular reason. Once we have other examples, it can be deleted.
+#if 0
+		FCoreDelegates::GetHotfixDelegate(EHotfixDelegates::Test).BindRaw(this, &FWindowsTargetPlatformSettingsModule::HotfixTest);
+#endif
+	}
+
+	virtual void ShutdownModule() override
+	{
+		// this is an example of a hotfix, declared here for no particular reason. Once we have other examples, it can be deleted.
+#if 0
+		FCoreDelegates::GetHotfixDelegate(EHotfixDelegates::Test).Unbind();
+#endif
+
+		ISettingsModule* SettingsModule = FModuleManager::GetModulePtr<ISettingsModule>("Settings");
+	}
+
+private:
+	TMap<FString, ITargetPlatformSettings*> PlatformNameToPlatformSettings;
+	ITargetPlatformSettings* PlatformSettingsCookedEditor;
+	ITargetPlatformSettings* PlatformSettingsCookedCooker;
+};
+
+
+#undef LOCTEXT_NAMESPACE
+
+
+IMPLEMENT_MODULE(FWindowsTargetPlatformSettingsModule, WindowsTargetPlatformSettings);
diff --git a/Engine/Source/Developer/Windows/WindowsTargetPlatformSettings/Public/GenericWindowsTargetPlatformSettings.h b/Engine/Source/Developer/Windows/WindowsTargetPlatformSettings/Public/GenericWindowsTargetPlatformSettings.h
new file mode 100644
--- /dev/null
+++ b/Engine/Source/Developer/Windows/WindowsTargetPlatformSettings/Public/GenericWindowsTargetPlatformSettings.h
@@ -0,0 +1,273 @@
+// Copyright Epic Games, Inc. All Rights Reserved.
+
+#pragma once
+
+#include "CoreMinimal.h"
+#include "Common/TargetPlatformSettingsBase.h"
+#include "Misc/ConfigCacheIni.h"
+#if WITH_ENGINE
+	#include "StaticMeshResources.h"
+	#include "RHI.h"
+	#include "DataDrivenShaderPlatformInfo.h"
+#endif // WITH_ENGINE
+#include "HAL/IConsoleManager.h"
+
+#define LOCTEXT_NAMESPACE "TGenericWindowsTargetPlatformSettings"
+
+/**
+ * Template for Windows target platforms settings
+ */
+
+template<typename TProperties>
+class TGenericWindowsTargetPlatformSettings : public TTargetPlatformSettingsBase<TProperties>
+{
+public:
+	typedef TTargetPlatformSettingsBase<TProperties> TSuper;
+
+	/**
+	 * Default constructor.
+	 */
+	TGenericWindowsTargetPlatformSettings( )
+	{
+
+	#if WITH_ENGINE
+		TextureLODSettings = nullptr; // These are registered by the device profile system.
+		StaticMeshLODSettings.Initialize(this);
+
+
+		// Get the Target RHIs for this platform, we do not always want all those that are supported.
+		TArray<FName> TargetedShaderFormats;
+		TGenericWindowsTargetPlatformSettings::GetAllTargetedShaderFormats(TargetedShaderFormats);
+
+		static FName NAME_PCD3D_SM6(TEXT("PCD3D_SM6"));
+		static FName NAME_PCD3D_SM5(TEXT("PCD3D_SM5"));
+		static FName NAME_VULKAN_SM5(TEXT("SF_VULKAN_SM5"));
+		static FName NAME_VULKAN_SM6(TEXT("SF_VULKAN_SM6"));
+		static FName NAME_PCD3D_ES3_1(TEXT("PCD3D_ES31"));
+		static FName NAME_SF_VULKAN_ES31(TEXT("SF_VULKAN_ES31"));
+		static FName NAME_OPENGL_150_ES3_1(TEXT("GLSL_150_ES31"));
+
+		// If we are targeting ES3.1, we also must cook encoded HDR reflection captures
+		bRequiresEncodedHDRReflectionCaptures =	TargetedShaderFormats.Contains(NAME_SF_VULKAN_ES31)
+												|| TargetedShaderFormats.Contains(NAME_OPENGL_150_ES3_1)
+												|| TargetedShaderFormats.Contains(NAME_PCD3D_ES3_1);
+
+	#endif
+	}
+
+public:
+
+	//~ Begin ITargetPlatform Interface
+
+	virtual bool SupportsFeature( ETargetPlatformFeatures Feature ) const override
+	{
+		// we currently do not have a build target for WindowsServer
+		if (Feature == ETargetPlatformFeatures::Packaging)
+		{
+			return (TProperties::HasEditorOnlyData() || !TProperties::IsServerOnly());
+		}
+
+		if ( Feature == ETargetPlatformFeatures::ShouldSplitPaksIntoSmallerSizes )
+		{
+			return TProperties::IsClientOnly();
+		}
+
+		if (Feature == ETargetPlatformFeatures::CanCookPackages)
+		{
+			return TProperties::HasEditorOnlyData();
+		}
+
+		if (Feature == ETargetPlatformFeatures::MobileRendering)
+		{
+			static bool bCachedSupportsMobileRendering = false;
+#if WITH_ENGINE
+			static bool bHasCachedValue = false;
+			if (!bHasCachedValue)
+			{
+				TArray<FName> TargetedShaderFormats;
+				GetAllTargetedShaderFormats(TargetedShaderFormats);
+
+				for (const FName& Format : TargetedShaderFormats)
+				{
+					if (IsMobilePlatform(ShaderFormatToLegacyShaderPlatform(Format)))
+					{
+						bCachedSupportsMobileRendering = true;
+						break;
+					}
+				}
+				bHasCachedValue = true;
+			}
+#endif
+
+			return bCachedSupportsMobileRendering;
+		}
+
+		return TSuper::SupportsFeature(Feature);
+	}
+
+	virtual void GetAllPossibleShaderFormats( TArray<FName>& OutFormats ) const override
+	{
+		// no shaders needed for dedicated server target
+		if (!TProperties::IsServerOnly())
+		{
+			static FName NAME_PCD3D_SM6(TEXT("PCD3D_SM6"));
+			static FName NAME_PCD3D_SM5(TEXT("PCD3D_SM5"));
+			static FName NAME_VULKAN_ES31(TEXT("SF_VULKAN_ES31"));
+			static FName NAME_OPENGL_150_ES3_1(TEXT("GLSL_150_ES31"));
+			static FName NAME_VULKAN_SM5(TEXT("SF_VULKAN_SM5"));
+			static FName NAME_VULKAN_SM6(TEXT("SF_VULKAN_SM6"));
+			static FName NAME_PCD3D_ES3_1(TEXT("PCD3D_ES31"));
+
+			OutFormats.AddUnique(NAME_PCD3D_SM5);
+			OutFormats.AddUnique(NAME_PCD3D_SM6);
+			OutFormats.AddUnique(NAME_VULKAN_ES31);
+			OutFormats.AddUnique(NAME_OPENGL_150_ES3_1);
+			OutFormats.AddUnique(NAME_VULKAN_SM5);
+			OutFormats.AddUnique(NAME_VULKAN_SM6);
+			OutFormats.AddUnique(NAME_PCD3D_ES3_1);
+		}
+	}
+
+	void GetAllTargetedShaderFormatsInternal(TArrayView<TCHAR const*> RelevantSettings, TArray<FName>& OutFormats) const
+	{
+		TArray<FString> TargetedShaderFormats;
+
+		for (const TCHAR* Name : RelevantSettings)
+		{
+			TArray<FString> NewTargetedShaderFormats;
+			GConfig->GetArray(TEXT("/Script/WindowsTargetPlatform.WindowsTargetSettings"), Name, NewTargetedShaderFormats, GEngineIni);
+
+			for (const FString& NewShaderFormat : NewTargetedShaderFormats)
+			{
+				TargetedShaderFormats.AddUnique(NewShaderFormat);
+			}
+		}
+		
+		// Gather the list of Target RHIs and filter out any that may be invalid.
+		TArray<FName> PossibleShaderFormats;
+		GetAllPossibleShaderFormats(PossibleShaderFormats);
+
+		for (int32 ShaderFormatIdx = TargetedShaderFormats.Num() - 1; ShaderFormatIdx >= 0; ShaderFormatIdx--)
+		{
+			FString ShaderFormat = TargetedShaderFormats[ShaderFormatIdx];
+			if (PossibleShaderFormats.Contains(FName(*ShaderFormat)) == false)
+			{
+				TargetedShaderFormats.RemoveAt(ShaderFormatIdx);
+			}
+		}
+
+		for(const FString& ShaderFormat : TargetedShaderFormats)
+		{
+			OutFormats.AddUnique(FName(*ShaderFormat));
+		}
+	}
+
+	virtual void GetAllTargetedShaderFormats( TArray<FName>& OutFormats ) const override 
+	{
+		// Get the Target RHIs for this platform, we do not always want all those that are supported. (reload in case user changed in the editor)
+
+		TCHAR const* RelevantSettings[] = 
+		{
+			TEXT("TargetedRHIs"),
+			TEXT("D3D12TargetedShaderFormats"),
+			TEXT("D3D11TargetedShaderFormats"),
+			TEXT("VulkanTargetedShaderFormats")
+		};
+
+		GetAllTargetedShaderFormatsInternal(RelevantSettings, OutFormats);
+	}
+
+	virtual void GetRayTracingShaderFormats( TArray<FName>& OutFormats ) const override 
+	{
+		if (UsesRayTracing())
+		{
+			TCHAR const* RelevantSettings[] = 
+			{
+				TEXT("VulkanTargetedShaderFormats")
+			};
+
+			GetAllTargetedShaderFormatsInternal(RelevantSettings, OutFormats);
+
+			// We always support ray tracing shaders when cooking for D3D12 SM6, however we may skip them for SM5 based on project settings.
+			OutFormats.AddUnique(FName(TEXT("PCD3D_SM6")));
+
+			static IConsoleVariable* RequireSM6CVar = IConsoleManager::Get().FindConsoleVariable(TEXT("r.RayTracing.RequireSM6"));
+			const bool bRequireSM6 = RequireSM6CVar && RequireSM6CVar->GetBool();
+			if (!bRequireSM6)
+			{
+				OutFormats.AddUnique(FName(TEXT("PCD3D_SM5")));
+			}
+		}
+	}
+
+#if WITH_ENGINE
+	virtual void GetReflectionCaptureFormats(TArray<FName>& OutFormats) const override
+	{
+		if (bRequiresEncodedHDRReflectionCaptures)
+		{
+			OutFormats.Add(FName(TEXT("EncodedHDR")));
+		}
+
+		OutFormats.Add(FName(TEXT("FullHDR")));
+	}
+
+	virtual void GetShaderFormatModuleHints(TArray<FName>& OutModuleNames) const override
+	{
+		OutModuleNames.Add(TEXT("ShaderFormatD3D"));
+		OutModuleNames.Add(TEXT("ShaderFormatOpenGL"));
+		OutModuleNames.Add(TEXT("VulkanShaderFormat"));
+	}
+
+	virtual const class FStaticMeshLODSettings& GetStaticMeshLODSettings( ) const override
+	{
+		return StaticMeshLODSettings;
+	}
+
+
+	virtual const UTextureLODSettings& GetTextureLODSettings() const override
+	{
+		return *TextureLODSettings;
+	}
+
+	virtual void RegisterTextureLODSettings(const UTextureLODSettings* InTextureLODSettings) override
+	{
+		TextureLODSettings = InTextureLODSettings;
+	}
+
+#endif //WITH_ENGINE
+
+	virtual bool UsesDistanceFields() const override
+	{
+		bool bEnableDistanceFields = false;
+		GConfig->GetBool(TEXT("/Script/WindowsTargetPlatform.WindowsTargetSettings"), TEXT("bEnableDistanceFields"), bEnableDistanceFields, GEngineIni);
+
+		return bEnableDistanceFields && TSuper::UsesDistanceFields();
+	}
+
+	virtual bool UsesRayTracing() const override
+	{
+		bool bEnableRayTracing = false;
+		GConfig->GetBool(TEXT("/Script/WindowsTargetPlatform.WindowsTargetSettings"), TEXT("bEnableRayTracing"), bEnableRayTracing, GEngineIni);
+
+		return bEnableRayTracing && TSuper::UsesRayTracing();
+	}
+
+	//~ End ITargetPlatform Interface
+
+private:
+
+#if WITH_ENGINE
+	// Holds the texture LOD settings.
+	const UTextureLODSettings* TextureLODSettings;
+
+	// Holds static mesh LOD settings.
+	FStaticMeshLODSettings StaticMeshLODSettings;
+
+	// True if the project requires encoded HDR reflection captures
+	bool bRequiresEncodedHDRReflectionCaptures;
+
+#endif // WITH_ENGINE
+
+};
+
+#undef LOCTEXT_NAMESPACE
diff --git a/Engine/Source/Developer/Windows/WindowsTargetPlatformSettings/Public/IWindowsTargetPlatformSettingsModule.h b/Engine/Source/Developer/Windows/WindowsTargetPlatformSettings/Public/IWindowsTargetPlatformSettingsModule.h
new file mode 100644
--- /dev/null
+++ b/Engine/Source/Developer/Windows/WindowsTargetPlatformSettings/Public/IWindowsTargetPlatformSettingsModule.h
@@ -0,0 +1,19 @@
+// Copyright Epic Games, Inc. All Rights Reserved.
+
+#pragma  once 
+
+/*------------------------------------------------------------------------------------
+IWindowsTargetPlatformSettingsModule interface
+------------------------------------------------------------------------------------*/
+
+#include "CoreMinimal.h"
+#include "Interfaces/ITargetPlatformSettingsModule.h"
+#include "GenericWindowsTargetPlatformSettings.h"
+
+class IWindowsTargetPlatformSettingsModule : public ITargetPlatformSettingsModule
+{
+public:
+	virtual void GetPlatformSettingsMaps(TMap<FString, ITargetPlatformSettings*>& OutMap) = 0;
+	virtual ITargetPlatformSettings* GetCookedEditorPlatformSettings() = 0;
+	virtual ITargetPlatformSettings* GetCookedCookerPlatformSettings() = 0;
+};
diff --git a/Engine/Source/Developer/Windows/WindowsTargetPlatformSettings/WindowsTargetPlatformSettings.Build.cs b/Engine/Source/Developer/Windows/WindowsTargetPlatformSettings/WindowsTargetPlatformSettings.Build.cs
new file mode 100644
--- /dev/null
+++ b/Engine/Source/Developer/Windows/WindowsTargetPlatformSettings/WindowsTargetPlatformSettings.Build.cs
@@ -0,0 +1,48 @@
+// Copyright Epic Games, Inc. All Rights Reserved.
+
+using UnrealBuildTool;
+
+public class WindowsTargetPlatformSettings : ModuleRules
+{
+	public WindowsTargetPlatformSettings(ReadOnlyTargetRules Target) : base(Target)
+	{
+		// We need a short name here since this can run afoul very easily of the `MAX_PATH` limit when
+		// combined with building other targets that make use of this.
+		ShortName = "WinTPSet";
+
+		PrivateDependencyModuleNames.AddRange(
+			new string[] {
+				"Core",
+				"CoreUObject",
+				"TargetPlatform",
+				"DesktopPlatform",
+                "AudioPlatformConfiguration",
+            }
+		);
+
+		PrivateIncludePathModuleNames.AddRange(
+			new string[] {
+				"Settings"
+			}
+		);
+
+		PublicIncludePathModuleNames.AddRange(
+			new string[] {
+				"AudioPlatformConfiguration"
+			}
+		);
+
+		// compile with Engine
+		if (Target.bCompileAgainstEngine)
+		{
+			PublicIncludePathModuleNames.Add("Engine");
+			PrivateDependencyModuleNames.AddRange( new string[] {
+				"Engine", 
+				"RHI",
+				"CookedEditor",
+				}
+			);
+            PrivateIncludePathModuleNames.Add("TextureCompressor");
+        }
+    }
+}
diff --git a/Engine/Source/Programs/UnrealBuildTool/Platform/Linux/UEBuildLinux.cs b/Engine/Source/Programs/UnrealBuildTool/Platform/Linux/UEBuildLinux.cs
--- a/Engine/Source/Programs/UnrealBuildTool/Platform/Linux/UEBuildLinux.cs
+++ b/Engine/Source/Programs/UnrealBuildTool/Platform/Linux/UEBuildLinux.cs
@@ -392,7 +392,11 @@ namespace UnrealBuildTool
 					{
 						if (Target.bBuildDeveloperTools)
 						{
+							Rules.DynamicallyLoadedModuleNames.Add("LinuxTargetPlatformSettings");
+							Rules.DynamicallyLoadedModuleNames.Add("LinuxTargetPlatformControls");
 							Rules.DynamicallyLoadedModuleNames.Add("LinuxTargetPlatform");
+							Rules.DynamicallyLoadedModuleNames.Add("LinuxArm64TargetPlatformSettings");
+							Rules.DynamicallyLoadedModuleNames.Add("LinuxArm64TargetPlatformControls");
 							Rules.DynamicallyLoadedModuleNames.Add("LinuxArm64TargetPlatform");
 						}
 					}
@@ -401,7 +405,11 @@ namespace UnrealBuildTool
 				// allow standalone tools to use targetplatform modules, without needing Engine
 				if (Target.bForceBuildTargetPlatforms && ModuleName == "TargetPlatform")
 				{
+					Rules.DynamicallyLoadedModuleNames.Add("LinuxTargetPlatformSettings");
+					Rules.DynamicallyLoadedModuleNames.Add("LinuxTargetPlatformControls");
 					Rules.DynamicallyLoadedModuleNames.Add("LinuxTargetPlatform");
+					Rules.DynamicallyLoadedModuleNames.Add("LinuxArm64TargetPlatformSettings");
+					Rules.DynamicallyLoadedModuleNames.Add("LinuxArm64TargetPlatformControls");
 					Rules.DynamicallyLoadedModuleNames.Add("LinuxArm64TargetPlatform");
 				}
 			}
@@ -431,7 +439,11 @@ namespace UnrealBuildTool
 			{
 				if (Target.bForceBuildTargetPlatforms)
 				{
+					Rules.DynamicallyLoadedModuleNames.Add("LinuxTargetPlatformSettings");
+					Rules.DynamicallyLoadedModuleNames.Add("LinuxTargetPlatformControls");
 					Rules.DynamicallyLoadedModuleNames.Add("LinuxTargetPlatform");
+					Rules.DynamicallyLoadedModuleNames.Add("LinuxArm64TargetPlatformSettings");
+					Rules.DynamicallyLoadedModuleNames.Add("LinuxArm64TargetPlatformControls");
 					Rules.DynamicallyLoadedModuleNames.Add("LinuxArm64TargetPlatform");
 				}
 
diff --git a/Engine/Source/Programs/UnrealBuildTool/Platform/Mac/UEBuildMac.cs b/Engine/Source/Programs/UnrealBuildTool/Platform/Mac/UEBuildMac.cs
--- a/Engine/Source/Programs/UnrealBuildTool/Platform/Mac/UEBuildMac.cs
+++ b/Engine/Source/Programs/UnrealBuildTool/Platform/Mac/UEBuildMac.cs
@@ -458,7 +458,10 @@ namespace UnrealBuildTool
 		public override void ModifyModuleRulesForOtherPlatform(string ModuleName, ModuleRules Rules, ReadOnlyTargetRules Target)
 		{
 			// don't do any target platform stuff if SDK is not available
-			if (!UEBuildPlatform.IsPlatformAvailableForTarget(Platform, Target))
+			bool bIsPlatformAvailableForTarget = UEBuildPlatform.IsPlatformAvailableForTarget(Platform, Target, bIgnoreSDKCheck: true);
+			bool bIsPlatformAvailableForTargetWithSDK = UEBuildPlatform.IsPlatformAvailableForTarget(Platform, Target);
+			
+			if (!bIsPlatformAvailableForTarget)
 			{
 				return;
 			}
@@ -468,7 +471,11 @@ namespace UnrealBuildTool
 				// because remote IOS building needs the new XcodeProject Settings to show up in the editor, we bring in the Mac bits that expose it
 				if (ModuleName == "Engine")
 				{
-					Rules.DynamicallyLoadedModuleNames.AddAll("MacTargetPlatform", "MacPlatformEditor");
+					Rules.DynamicallyLoadedModuleNames.Add("MacTargetPlatformSettings");
+					if (bIsPlatformAvailableForTargetWithSDK)
+					{
+						Rules.DynamicallyLoadedModuleNames.AddAll("MacTargetPlatform", "MacTargetPlatformControls", "MacPlatformEditor");
+					}
 				}
 			}
 		}
@@ -530,6 +537,8 @@ namespace UnrealBuildTool
 				if (Target.bForceBuildTargetPlatforms)
 				{
 					Rules.DynamicallyLoadedModuleNames.Add("MacTargetPlatform");
+					Rules.DynamicallyLoadedModuleNames.Add("MacTargetPlatformSettings");
+					Rules.DynamicallyLoadedModuleNames.Add("MacTargetPlatformControls");
 				}
 
 				if (bBuildShaderFormats)
diff --git a/Engine/Source/Programs/UnrealBuildTool/Platform/Windows/UEBuildWindows.cs b/Engine/Source/Programs/UnrealBuildTool/Platform/Windows/UEBuildWindows.cs
--- a/Engine/Source/Programs/UnrealBuildTool/Platform/Windows/UEBuildWindows.cs
+++ b/Engine/Source/Programs/UnrealBuildTool/Platform/Windows/UEBuildWindows.cs
@@ -1541,6 +1541,8 @@ namespace UnrealBuildTool
 				if (Target.bForceBuildTargetPlatforms)
 				{
 					Rules.DynamicallyLoadedModuleNames.Add("WindowsTargetPlatform");
+					Rules.DynamicallyLoadedModuleNames.Add("WindowsTargetPlatformSettings");
+					Rules.DynamicallyLoadedModuleNames.Add("WindowsTargetPlatformControls");
 				}
 
 				if (bBuildShaderFormats)
