- Set all *TargetPlatform.Build.cs's to set the new SDKVersionRelevantPlatforms list so that per-project SDK version overrides will compile unique versions of those modules when that platform's SDK version of overridden

diff --git a/Engine/Platforms/VisionOS/Source/Developer/VisionOSTargetPlatform/VisionOSTargetPlatform.Build.cs b/Engine/Platforms/VisionOS/Source/Developer/VisionOSTargetPlatform/VisionOSTargetPlatform.Build.cs
--- a/Engine/Platforms/VisionOS/Source/Developer/VisionOSTargetPlatform/VisionOSTargetPlatform.Build.cs
+++ b/Engine/Platforms/VisionOS/Source/Developer/VisionOSTargetPlatform/VisionOSTargetPlatform.Build.cs
@@ -7,7 +7,8 @@ public class VisionOSTargetPlatform : ModuleRules
 	public VisionOSTargetPlatform(ReadOnlyTargetRules Target) : base(Target)
 	{
 		BinariesSubFolder = "IOS";
-		
+		SDKVersionRelevantPlatforms.Add(UnrealTargetPlatform.VisionOS);
+
 		PrivateDependencyModuleNames.AddRange(
 			new string[]
 			{
diff --git a/Engine/Source/Developer/Android/AndroidTargetPlatform/AndroidTargetPlatform.Build.cs b/Engine/Source/Developer/Android/AndroidTargetPlatform/AndroidTargetPlatform.Build.cs
--- a/Engine/Source/Developer/Android/AndroidTargetPlatform/AndroidTargetPlatform.Build.cs
+++ b/Engine/Source/Developer/Android/AndroidTargetPlatform/AndroidTargetPlatform.Build.cs
@@ -8,8 +8,9 @@ public class AndroidTargetPlatform : ModuleRules
 	{
 		UnsafeTypeCastWarningLevel = WarningLevel.Error;
 		BinariesSubFolder = "Android";
+		SDKVersionRelevantPlatforms.Add(UnrealTargetPlatform.Android);
 
-        PrivateDependencyModuleNames.AddRange(
+		PrivateDependencyModuleNames.AddRange(
 			new string[]
 			{
 				"Core",
diff --git a/Engine/Source/Developer/Android/AndroidTargetPlatformControls/AndroidTargetPlatformControls.Build.cs b/Engine/Source/Developer/Android/AndroidTargetPlatformControls/AndroidTargetPlatformControls.Build.cs
--- a/Engine/Source/Developer/Android/AndroidTargetPlatformControls/AndroidTargetPlatformControls.Build.cs
+++ b/Engine/Source/Developer/Android/AndroidTargetPlatformControls/AndroidTargetPlatformControls.Build.cs
@@ -8,6 +8,8 @@ public class AndroidTargetPlatformControls : ModuleRules
 	{
 		UnsafeTypeCastWarningLevel = WarningLevel.Error;
 		BinariesSubFolder = "Android";
+		SDKVersionRelevantPlatforms.Add(UnrealTargetPlatform.Android);
+
 		// We need a short name here since this can run afoul very easily of the `MAX_PATH` limit when
 		// combined with building other targets that make use of this.
 		ShortName = "AndTPCon";
diff --git a/Engine/Source/Developer/IOS/IOSTargetPlatform/IOSTargetPlatform.Build.cs b/Engine/Source/Developer/IOS/IOSTargetPlatform/IOSTargetPlatform.Build.cs
--- a/Engine/Source/Developer/IOS/IOSTargetPlatform/IOSTargetPlatform.Build.cs
+++ b/Engine/Source/Developer/IOS/IOSTargetPlatform/IOSTargetPlatform.Build.cs
@@ -7,6 +7,7 @@ public class IOSTargetPlatform : ModuleRules
 	public IOSTargetPlatform(ReadOnlyTargetRules Target) : base(Target)
 	{
 		BinariesSubFolder = "IOS";
+		SDKVersionRelevantPlatforms.Add(UnrealTargetPlatform.IOS);
 
 		PrivateDependencyModuleNames.AddRange(
 			new string[]
@@ -24,6 +25,7 @@ public class IOSTargetPlatform : ModuleRules
 			}
 		);
 
+
 		PrivateIncludePathModuleNames.AddRange(
 			new string[] {
 				"MessagingCommon",
diff --git a/Engine/Source/Developer/IOS/TVOSTargetPlatform/TVOSTargetPlatform.Build.cs b/Engine/Source/Developer/IOS/TVOSTargetPlatform/TVOSTargetPlatform.Build.cs
--- a/Engine/Source/Developer/IOS/TVOSTargetPlatform/TVOSTargetPlatform.Build.cs
+++ b/Engine/Source/Developer/IOS/TVOSTargetPlatform/TVOSTargetPlatform.Build.cs
@@ -7,7 +7,8 @@ public class TVOSTargetPlatform : ModuleRules
 	public TVOSTargetPlatform(ReadOnlyTargetRules Target) : base(Target)
 	{
 		BinariesSubFolder = "IOS";
-		
+		SDKVersionRelevantPlatforms.Add(UnrealTargetPlatform.TVOS);
+
 		PrivateDependencyModuleNames.AddRange(
 			new string[]
 			{
diff --git a/Engine/Source/Developer/Linux/LinuxArm64TargetPlatform/LinuxArm64TargetPlatform.Build.cs b/Engine/Source/Developer/Linux/LinuxArm64TargetPlatform/LinuxArm64TargetPlatform.Build.cs
--- a/Engine/Source/Developer/Linux/LinuxArm64TargetPlatform/LinuxArm64TargetPlatform.Build.cs
+++ b/Engine/Source/Developer/Linux/LinuxArm64TargetPlatform/LinuxArm64TargetPlatform.Build.cs
@@ -7,8 +7,9 @@ public class LinuxArm64TargetPlatform : ModuleRules
     public LinuxArm64TargetPlatform(ReadOnlyTargetRules Target) : base(Target)
     {
         BinariesSubFolder = "LinuxArm64";
+		SDKVersionRelevantPlatforms.Add(UnrealTargetPlatform.LinuxArm64);
 
-        PrivateDependencyModuleNames.AddRange(
+		PrivateDependencyModuleNames.AddRange(
             new string[] {
 				"Core",
 				"DesktopPlatform",
diff --git a/Engine/Source/Developer/Linux/LinuxArm64TargetPlatformControls/LinuxArm64TargetPlatformControls.Build.cs b/Engine/Source/Developer/Linux/LinuxArm64TargetPlatformControls/LinuxArm64TargetPlatformControls.Build.cs
--- a/Engine/Source/Developer/Linux/LinuxArm64TargetPlatformControls/LinuxArm64TargetPlatformControls.Build.cs
+++ b/Engine/Source/Developer/Linux/LinuxArm64TargetPlatformControls/LinuxArm64TargetPlatformControls.Build.cs
@@ -7,6 +7,8 @@ public class LinuxArm64TargetPlatformControls : ModuleRules
     public LinuxArm64TargetPlatformControls(ReadOnlyTargetRules Target) : base(Target)
     {
         BinariesSubFolder = "LinuxArm64";
+		SDKVersionRelevantPlatforms.Add(UnrealTargetPlatform.LinuxArm64);
+
 		// We need a short name here since this can run afoul very easily of the `MAX_PATH` limit when
 		// combined with building other targets that make use of this.
 		ShortName = "LArm64TPCon";
diff --git a/Engine/Source/Developer/Linux/LinuxTargetPlatform/LinuxTargetPlatform.Build.cs b/Engine/Source/Developer/Linux/LinuxTargetPlatform/LinuxTargetPlatform.Build.cs
--- a/Engine/Source/Developer/Linux/LinuxTargetPlatform/LinuxTargetPlatform.Build.cs
+++ b/Engine/Source/Developer/Linux/LinuxTargetPlatform/LinuxTargetPlatform.Build.cs
@@ -7,6 +7,7 @@ public class LinuxTargetPlatform : ModuleRules
     public LinuxTargetPlatform(ReadOnlyTargetRules Target) : base(Target)
 	{
         BinariesSubFolder = "Linux";
+		SDKVersionRelevantPlatforms.AddRange(Utils.GetPlatformsInGroup(UnrealPlatformGroup.Linux));
 
 		PrivateDependencyModuleNames.AddRange(
             new string[] {
diff --git a/Engine/Source/Developer/Linux/LinuxTargetPlatformControls/LinuxTargetPlatformControls.Build.cs b/Engine/Source/Developer/Linux/LinuxTargetPlatformControls/LinuxTargetPlatformControls.Build.cs
--- a/Engine/Source/Developer/Linux/LinuxTargetPlatformControls/LinuxTargetPlatformControls.Build.cs
+++ b/Engine/Source/Developer/Linux/LinuxTargetPlatformControls/LinuxTargetPlatformControls.Build.cs
@@ -7,6 +7,8 @@ public class LinuxTargetPlatformControls : ModuleRules
     public LinuxTargetPlatformControls(ReadOnlyTargetRules Target) : base(Target)
 	{
         BinariesSubFolder = "Linux";
+		SDKVersionRelevantPlatforms.AddRange(Utils.GetPlatformsInGroup(UnrealPlatformGroup.Linux));
+
 		// We need a short name here since this can run afoul very easily of the `MAX_PATH` limit when
 		// combined with building other targets that make use of this.
 		ShortName = "LinuxTPCon";
diff --git a/Engine/Source/Developer/Mac/MacTargetPlatform/MacTargetPlatform.Build.cs b/Engine/Source/Developer/Mac/MacTargetPlatform/MacTargetPlatform.Build.cs
--- a/Engine/Source/Developer/Mac/MacTargetPlatform/MacTargetPlatform.Build.cs
+++ b/Engine/Source/Developer/Mac/MacTargetPlatform/MacTargetPlatform.Build.cs
@@ -6,6 +6,8 @@ public class MacTargetPlatform : ModuleRules
 {
 	public MacTargetPlatform(ReadOnlyTargetRules Target) : base(Target)
 	{
+		SDKVersionRelevantPlatforms.Add(UnrealTargetPlatform.Mac);
+
 		PrivateDependencyModuleNames.AddRange(
 			new string[] {
 				"Core",
diff --git a/Engine/Source/Developer/Mac/MacTargetPlatformControls/MacTargetPlatformControls.Build.cs b/Engine/Source/Developer/Mac/MacTargetPlatformControls/MacTargetPlatformControls.Build.cs
--- a/Engine/Source/Developer/Mac/MacTargetPlatformControls/MacTargetPlatformControls.Build.cs
+++ b/Engine/Source/Developer/Mac/MacTargetPlatformControls/MacTargetPlatformControls.Build.cs
@@ -6,6 +6,8 @@ public class MacTargetPlatformControls : ModuleRules
 {
 	public MacTargetPlatformControls(ReadOnlyTargetRules Target) : base(Target)
 	{
+		SDKVersionRelevantPlatforms.Add(UnrealTargetPlatform.Mac);
+
 		// We need a short name here since this can run afoul very easily of the `MAX_PATH` limit when
 		// combined with building other targets that make use of this.
 		ShortName = "MacTPCon";
diff --git a/Engine/Source/Developer/Windows/WindowsTargetPlatfomControls/WindowsTargetPlatformControls.Build.cs b/Engine/Source/Developer/Windows/WindowsTargetPlatfomControls/WindowsTargetPlatformControls.Build.cs
--- a/Engine/Source/Developer/Windows/WindowsTargetPlatfomControls/WindowsTargetPlatformControls.Build.cs
+++ b/Engine/Source/Developer/Windows/WindowsTargetPlatfomControls/WindowsTargetPlatformControls.Build.cs
@@ -6,6 +6,8 @@ public class WindowsTargetPlatformControls : ModuleRules
 {
 	public WindowsTargetPlatformControls(ReadOnlyTargetRules Target) : base(Target)
 	{
+		SDKVersionRelevantPlatforms.Add(UnrealTargetPlatform.Win64);
+
 		// We need a short name here since this can run afoul very easily of the `MAX_PATH` limit when
 		// combined with building other targets that make use of this.
 		ShortName = "WinTPCon";
diff --git a/Engine/Source/Developer/Windows/WindowsTargetPlatform/WindowsTargetPlatform.Build.cs b/Engine/Source/Developer/Windows/WindowsTargetPlatform/WindowsTargetPlatform.Build.cs
--- a/Engine/Source/Developer/Windows/WindowsTargetPlatform/WindowsTargetPlatform.Build.cs
+++ b/Engine/Source/Developer/Windows/WindowsTargetPlatform/WindowsTargetPlatform.Build.cs
@@ -6,7 +6,9 @@ public class WindowsTargetPlatform : ModuleRules
 {
 	public WindowsTargetPlatform(ReadOnlyTargetRules Target) : base(Target)
 	{
-        PrivateDependencyModuleNames.AddRange(
+		SDKVersionRelevantPlatforms.Add(UnrealTargetPlatform.Win64);
+		
+		PrivateDependencyModuleNames.AddRange(
 			new string[] {
 				"Core",
 				"TargetPlatform",
diff --git a/Engine/Source/Programs/UnrealBuildTool/System/Utils.cs b/Engine/Source/Programs/UnrealBuildTool/System/Utils.cs
--- a/Engine/Source/Programs/UnrealBuildTool/System/Utils.cs
+++ b/Engine/Source/Programs/UnrealBuildTool/System/Utils.cs
@@ -517,6 +517,15 @@ namespace UnrealBuildTool
 			throw new ArgumentException(String.Format("'{0}' is not a valid value for UnrealPlatformClass", (int)Class));
 		}
 
+		/// <summary>
+		/// Retrieve the list of platforms in this group (if any)
+		/// </summary>
+		public static List<UnrealTargetPlatform> GetPlatformsInGroup(UnrealPlatformGroup Group)
+		{
+			return UEBuildPlatform.GetPlatformsInGroup(Group);
+		}
+
+
 		/// <summary>
 		/// Given a list of supported platforms, returns a list of names of platforms that should not be supported
 		/// </summary>
