Fix Arm64TP

diff --git a/Engine/Source/Developer/Linux/LinuxArm64TargetPlatform/LinuxArm64TargetPlatform.Build.cs b/Engine/Source/Developer/Linux/LinuxArm64TargetPlatform/LinuxArm64TargetPlatform.Build.cs
--- a/Engine/Source/Developer/Linux/LinuxArm64TargetPlatform/LinuxArm64TargetPlatform.Build.cs
+++ b/Engine/Source/Developer/Linux/LinuxArm64TargetPlatform/LinuxArm64TargetPlatform.Build.cs
@@ -14,8 +14,6 @@ public class LinuxArm64TargetPlatform : ModuleRules
 				"Core",
 				"DesktopPlatform",
 				"TargetPlatform",
-				"LinuxArm64TargetPlatformSettings",
-				"LinuxArm64TargetPlatformControls",
 			}
         );
     }
diff --git a/Engine/Source/Developer/Linux/LinuxArm64TargetPlatformControls/LinuxArm64TargetPlatformControls.Build.cs b/Engine/Source/Developer/Linux/LinuxArm64TargetPlatformControls/LinuxArm64TargetPlatformControls.Build.cs
--- a/Engine/Source/Developer/Linux/LinuxArm64TargetPlatformControls/LinuxArm64TargetPlatformControls.Build.cs
+++ b/Engine/Source/Developer/Linux/LinuxArm64TargetPlatformControls/LinuxArm64TargetPlatformControls.Build.cs
@@ -19,8 +19,6 @@ public class LinuxArm64TargetPlatformControls : ModuleRules
 				"DesktopPlatform",
 				"TargetPlatform",
 				"LinuxArm64TargetPlatformSettings",
-				"LinuxTargetPlatformSettings",
-				"LinuxTargetPlatformControls",
 			}
         );
 
diff --git a/Engine/Source/Developer/Linux/LinuxArm64TargetPlatformSettings/LinuxArm64TargetPlatformSettings.Build.cs b/Engine/Source/Developer/Linux/LinuxArm64TargetPlatformSettings/LinuxArm64TargetPlatformSettings.Build.cs
--- a/Engine/Source/Developer/Linux/LinuxArm64TargetPlatformSettings/LinuxArm64TargetPlatformSettings.Build.cs
+++ b/Engine/Source/Developer/Linux/LinuxArm64TargetPlatformSettings/LinuxArm64TargetPlatformSettings.Build.cs
@@ -16,7 +16,6 @@ public class LinuxArm64TargetPlatformSettings : ModuleRules
 				"Core",
 				"DesktopPlatform",
 				"TargetPlatform",
-				"LinuxTargetPlatformSettings",
 			}
         );
 
