Ray Tracing: Fix typo in UI for Linux systems

#rnx

#Synced-CL 51240028

diff --git a/Engine/Source/Developer/Linux/LinuxTargetPlatformSettings/Classes/LinuxTargetSettings.h b/Engine/Source/Developer/Linux/LinuxTargetPlatformSettings/Classes/LinuxTargetSettings.h
--- a/Engine/Source/Developer/Linux/LinuxTargetPlatformSettings/Classes/LinuxTargetSettings.h
+++ b/Engine/Source/Developer/Linux/LinuxTargetPlatformSettings/Classes/LinuxTargetSettings.h
@@ -61,7 +61,7 @@ public:
 	UPROPERTY(config, meta = (DeprecatedProperty))
 	bool bEnableRayTracing_DEPRECATED;
 
-	/** Controls the level of ray tracing support for Windows. r.RayTracing cvar also needs to be enabled. */
+	/** Controls the level of ray tracing support for Linux. r.RayTracing cvar also needs to be enabled. */
 	UPROPERTY(config, EditAnywhere, Category = "Renderer", meta = (DisplayName = "Hardware Ray Tracing Mode"))
 	ETargetSettingsRayTracingRuntimeMode RayTracingMode;
 
