[Linux] Add support for ASTC/ETC2/DXT/Multi cook flavors for ARM

#13370
#13105

#Synced-CL 51597533

#Synced-CL 51597533

diff --git a/Engine/Source/Developer/Linux/LinuxArm64TargetPlatformControls/Private/LinuxArm64TargetPlatformControlsModule.cpp b/Engine/Source/Developer/Linux/LinuxArm64TargetPlatformControls/Private/LinuxArm64TargetPlatformControlsModule.cpp
--- a/Engine/Source/Developer/Linux/LinuxArm64TargetPlatformControls/Private/LinuxArm64TargetPlatformControlsModule.cpp
+++ b/Engine/Source/Developer/Linux/LinuxArm64TargetPlatformControls/Private/LinuxArm64TargetPlatformControlsModule.cpp
@@ -7,35 +7,90 @@
 #include "CoreMinimal.h"
 #include "Modules/ModuleManager.h"
 
-#include "Interfaces/ITargetPlatformControlsModule.h"
-
+#include "ILinuxArm64TargetPlatformControlsModule.h"
 #include "LinuxTargetPlatformSettings.h"
 #include "LinuxTargetPlatformControls.h"
 #include "ILinuxArm64TargetPlatformSettingsModule.h"
 
 /**
- * Module for the Linux target platforms controls
+ * Module for the Linux ARM64 target platform controls
  */
 class FLinuxArm64TargetPlatformControlsModule
-	: public ITargetPlatformControlsModule
+	: public ILinuxArm64TargetPlatformControlsModule
 {
 public:
 	virtual void GetTargetPlatformControls(TArray<ITargetPlatformControls*>& TargetPlatforms, FName& PlatformSettingsModuleName) override
 	{
-		TMap<FString, ITargetPlatformSettings*> PlatformNameToPlatformSettings;
 		ILinuxArm64TargetPlatformSettingsModule* ModuleSettings = FModuleManager::GetModulePtr<ILinuxArm64TargetPlatformSettingsModule>(PlatformSettingsModuleName);
 		if (ModuleSettings != nullptr)
 		{
+			// ensure this is only called once (matches Android's convention of guarding inside the settings-available branch)
+			check(SinglePlatforms.IsEmpty());
+			const FString GameTPName   = FLinuxPlatformProperties<false, false, false, true>::PlatformName();
+			const FString ServerTPName = FLinuxPlatformProperties<false, true,  false, true>::PlatformName();
+			const FString ClientTPName = FLinuxPlatformProperties<false, false, true,  true>::PlatformName();
+
 			TMap<FString, ITargetPlatformSettings*> OutMap;
 			ModuleSettings->GetPlatformSettingsMaps(OutMap);
-			ITargetPlatformControls* GameTP = new TLinuxTargetPlatformControls<FLinuxPlatformProperties<false, false, false, true>>(OutMap[FLinuxPlatformProperties<false, false, false, true>::PlatformName()]);
-			ITargetPlatformControls* ServerTP = new TLinuxTargetPlatformControls<FLinuxPlatformProperties<false, true, false, true>>(OutMap[FLinuxPlatformProperties<false, true, false, true>::PlatformName()]);
-			ITargetPlatformControls* ClientTP = new TLinuxTargetPlatformControls<FLinuxPlatformProperties<false, false, true, true>>(OutMap[FLinuxPlatformProperties<false, false, true, true>::PlatformName()]);
-			TargetPlatforms.Add(GameTP);
-			TargetPlatforms.Add(ServerTP);
-			TargetPlatforms.Add(ClientTP);
+
+			// --- Flavorless (default) ---
+			TargetPlatforms.Add(new TLinuxTargetPlatformControls<FLinuxPlatformProperties<false, false, false, true>>(OutMap[GameTPName]));
+			TargetPlatforms.Add(new TLinuxTargetPlatformControls<FLinuxPlatformProperties<false, true,  false, true>>(OutMap[ServerTPName]));
+			TargetPlatforms.Add(new TLinuxTargetPlatformControls<FLinuxPlatformProperties<false, false, true,  true>>(OutMap[ClientTPName]));
+
+			// --- DXT ---
+			SinglePlatforms.Add(new TLinuxDXTTargetPlatformControls<FLinuxPlatformProperties<false, false, false, true>>(OutMap[GameTPName   + TEXT("_DXT")]));
+			SinglePlatforms.Add(new TLinuxDXTTargetPlatformControls<FLinuxPlatformProperties<false, true,  false, true>>(OutMap[ServerTPName + TEXT("_DXT")]));
+			SinglePlatforms.Add(new TLinuxDXTTargetPlatformControls<FLinuxPlatformProperties<false, false, true,  true>>(OutMap[ClientTPName + TEXT("_DXT")]));
+
+			// --- ASTC ---
+			SinglePlatforms.Add(new TLinuxASTCTargetPlatformControls<FLinuxPlatformProperties<false, false, false, true>>(OutMap[GameTPName   + TEXT("_ASTC")]));
+			SinglePlatforms.Add(new TLinuxASTCTargetPlatformControls<FLinuxPlatformProperties<false, true,  false, true>>(OutMap[ServerTPName + TEXT("_ASTC")]));
+			SinglePlatforms.Add(new TLinuxASTCTargetPlatformControls<FLinuxPlatformProperties<false, false, true,  true>>(OutMap[ClientTPName + TEXT("_ASTC")]));
+
+			// --- ETC2 ---
+			SinglePlatforms.Add(new TLinuxETC2TargetPlatformControls<FLinuxPlatformProperties<false, false, false, true>>(OutMap[GameTPName   + TEXT("_ETC2")]));
+			SinglePlatforms.Add(new TLinuxETC2TargetPlatformControls<FLinuxPlatformProperties<false, true,  false, true>>(OutMap[ServerTPName + TEXT("_ETC2")]));
+			SinglePlatforms.Add(new TLinuxETC2TargetPlatformControls<FLinuxPlatformProperties<false, false, true,  true>>(OutMap[ClientTPName + TEXT("_ETC2")]));
+
+			TargetPlatforms.Append(SinglePlatforms);
+
+			// --- Multi ---
+			MultiGameTP   = new TLinuxMultiTargetPlatformControls<FLinuxPlatformProperties<false, false, false, true>>(OutMap[GameTPName   + TEXT("_Multi")]);
+			MultiServerTP = new TLinuxMultiTargetPlatformControls<FLinuxPlatformProperties<false, true,  false, true>>(OutMap[ServerTPName + TEXT("_Multi")]);
+			MultiClientTP = new TLinuxMultiTargetPlatformControls<FLinuxPlatformProperties<false, false, true,  true>>(OutMap[ClientTPName + TEXT("_Multi")]);
+			TargetPlatforms.Add(MultiGameTP);
+			TargetPlatforms.Add(MultiServerTP);
+			TargetPlatforms.Add(MultiClientTP);
+
+			// Load the multi format selection from config now that all single-format platforms are ready
+			NotifyMultiSelectedFormatsChanged();
 		}
 	}
+
+	virtual void NotifyMultiSelectedFormatsChanged() override
+	{
+		if (MultiGameTP && MultiServerTP && MultiClientTP)
+		{
+			MultiGameTP->LoadFormats(SinglePlatforms);
+			MultiServerTP->LoadFormats(SinglePlatforms);
+			MultiClientTP->LoadFormats(SinglePlatforms);
+		}
+	}
+
+private:
+	// Allocated platform controls live for the lifetime of this module. They are intentionally not deleted on unload,
+	// consistent with Android, iOS, Windows, Mac, and Linux modules which follow the same convention.
+	// GetTargetPlatformControls is guarded against multiple calls (see check above), so the objects pointed to by
+	// SinglePlatforms remain valid for the module's lifetime. LoadFormats() takes the array by value (a shallow copy
+	// of 9 pointers), which is safe because the pointed-to objects outlive the call.
+
+	/** Single-format cook targets (DXT, ASTC, ETC2), used to populate Multi on demand */
+	TArray<ITargetPlatformControls*> SinglePlatforms;
+
+	TLinuxMultiTargetPlatformControls<FLinuxPlatformProperties<false, false, false, true>>* MultiGameTP   = nullptr;
+	TLinuxMultiTargetPlatformControls<FLinuxPlatformProperties<false, true,  false, true>>* MultiServerTP = nullptr;
+	TLinuxMultiTargetPlatformControls<FLinuxPlatformProperties<false, false, true,  true>>* MultiClientTP = nullptr;
 };
 
 IMPLEMENT_MODULE(FLinuxArm64TargetPlatformControlsModule, LinuxArm64TargetPlatformControls);
diff --git a/Engine/Source/Developer/Linux/LinuxArm64TargetPlatformControls/Public/ILinuxArm64TargetPlatformControlsModule.h b/Engine/Source/Developer/Linux/LinuxArm64TargetPlatformControls/Public/ILinuxArm64TargetPlatformControlsModule.h
new file mode 100644
--- /dev/null
+++ b/Engine/Source/Developer/Linux/LinuxArm64TargetPlatformControls/Public/ILinuxArm64TargetPlatformControlsModule.h
@@ -0,0 +1,13 @@
+// Copyright Epic Games, Inc. All Rights Reserved.
+
+#pragma once
+
+#include "CoreMinimal.h"
+#include "Interfaces/ITargetPlatformControlsModule.h"
+
+class ILinuxArm64TargetPlatformControlsModule : public ITargetPlatformControlsModule
+{
+public:
+	// Called when the user changes the selected texture formats for the Multi cook target
+	virtual void NotifyMultiSelectedFormatsChanged() = 0;
+};
diff --git a/Engine/Source/Developer/Linux/LinuxArm64TargetPlatformSettings/Private/LinuxArm64TargetPlatformSettingsModule.cpp b/Engine/Source/Developer/Linux/LinuxArm64TargetPlatformSettings/Private/LinuxArm64TargetPlatformSettingsModule.cpp
--- a/Engine/Source/Developer/Linux/LinuxArm64TargetPlatformSettings/Private/LinuxArm64TargetPlatformSettingsModule.cpp
+++ b/Engine/Source/Developer/Linux/LinuxArm64TargetPlatformSettings/Private/LinuxArm64TargetPlatformSettingsModule.cpp
@@ -17,26 +17,81 @@ class FLinuxArm64TargetPlatformSettingsModule
 {
 public:
 
-	/** Destructor. */
-	~FLinuxArm64TargetPlatformSettingsModule()
-	{
-
-	}
-
+	// This follows the same single-call convention as all other platform settings modules (Android, iOS, Windows, Mac, Linux).
+	// The caller is responsible for ensuring it is only called once; no multiple-call guard is needed here.
 	virtual void GetTargetPlatformSettings(TArray<ITargetPlatformSettings*>& TargetPlatforms) override
 	{
+		const FString NoEditorTPName = FLinuxPlatformProperties<false, false, false, true>::PlatformName();
+		const FString ServerTPName = FLinuxPlatformProperties<false, true, false, true>::PlatformName();
+		const FString ClientTPName = FLinuxPlatformProperties<false, false, true, true>::PlatformName();
+	
 		// NoEditor TP
 		ITargetPlatformSettings* NoEditorTP = new TLinuxTargetPlatformSettings<FLinuxPlatformProperties<false, false, false, true>>();
 		TargetPlatforms.Add(NoEditorTP);
-		PlatformNameToPlatformSettings.Add(FLinuxPlatformProperties<false, false, false, true>::PlatformName(), NoEditorTP);
+		PlatformNameToPlatformSettings.Add(NoEditorTPName, NoEditorTP);
 		// Server TP
 		ITargetPlatformSettings* ServerTP = new TLinuxTargetPlatformSettings<FLinuxPlatformProperties<false, true, false, true>>();
 		TargetPlatforms.Add(ServerTP);
-		PlatformNameToPlatformSettings.Add(FLinuxPlatformProperties<false, true, false, true>::PlatformName(), ServerTP);
+		PlatformNameToPlatformSettings.Add(ServerTPName, ServerTP);
 		// Client TP
 		ITargetPlatformSettings* ClientTP = new TLinuxTargetPlatformSettings<FLinuxPlatformProperties<false, false, true, true>>();
 		TargetPlatforms.Add(ClientTP);
-		PlatformNameToPlatformSettings.Add(FLinuxPlatformProperties<false, false, true, true>::PlatformName(), ClientTP);
+		PlatformNameToPlatformSettings.Add(ClientTPName, ClientTP);
+
+		// --- DXT Variants ---
+		// NoEditor TP
+		NoEditorTP = new TLinux_DXTTargetPlatformSettings<FLinuxPlatformProperties<false, false, false, true>>();
+		TargetPlatforms.Add(NoEditorTP);
+		PlatformNameToPlatformSettings.Add(NoEditorTPName + TEXT("_DXT"), NoEditorTP);
+		// Server TP
+		ServerTP = new TLinux_DXTTargetPlatformSettings<FLinuxPlatformProperties<false, true, false, true>>();
+		TargetPlatforms.Add(ServerTP);
+		PlatformNameToPlatformSettings.Add(ServerTPName + TEXT("_DXT"), ServerTP);
+		// Client TP
+		ClientTP = new TLinux_DXTTargetPlatformSettings<FLinuxPlatformProperties<false, false, true, true>>();
+		TargetPlatforms.Add(ClientTP);
+		PlatformNameToPlatformSettings.Add(ClientTPName + TEXT("_DXT"), ClientTP);
+
+		// --- ASTC Variants ---
+		// NoEditor TP
+		NoEditorTP = new TLinux_ASTCTargetPlatformSettings<FLinuxPlatformProperties<false, false, false, true>>();
+		TargetPlatforms.Add(NoEditorTP);
+		PlatformNameToPlatformSettings.Add(NoEditorTPName + TEXT("_ASTC"), NoEditorTP);
+		// Server TP
+ 		ServerTP = new TLinux_ASTCTargetPlatformSettings<FLinuxPlatformProperties<false, true, false, true>>();
+ 		TargetPlatforms.Add(ServerTP);
+ 		PlatformNameToPlatformSettings.Add(ServerTPName + TEXT("_ASTC"), ServerTP);
+ 		// Client TP
+ 		ClientTP = new TLinux_ASTCTargetPlatformSettings<FLinuxPlatformProperties<false, false, true, true>>();
+ 		TargetPlatforms.Add(ClientTP);
+ 		PlatformNameToPlatformSettings.Add(ClientTPName + TEXT("_ASTC"), ClientTP);
+ 		 
+ 		// --- ETC2 Variants ---
+ 		NoEditorTP = new TLinux_ETC2TargetPlatformSettings<FLinuxPlatformProperties<false, false, false, true>>();
+ 		TargetPlatforms.Add(NoEditorTP);
+ 		PlatformNameToPlatformSettings.Add(NoEditorTPName + TEXT("_ETC2"), NoEditorTP);
+ 		// Server TP
+ 		ServerTP = new TLinux_ETC2TargetPlatformSettings<FLinuxPlatformProperties<false, true, false, true>>();
+ 		TargetPlatforms.Add(ServerTP);
+ 		PlatformNameToPlatformSettings.Add(ServerTPName + TEXT("_ETC2"), ServerTP);
+ 		// Client TP
+ 		ClientTP = new TLinux_ETC2TargetPlatformSettings<FLinuxPlatformProperties<false, false, true, true>>();
+ 		TargetPlatforms.Add(ClientTP);
+ 		PlatformNameToPlatformSettings.Add(ClientTPName + TEXT("_ETC2"), ClientTP);
+
+		// --- Multi Variants ---
+		// NoEditor TP
+		NoEditorTP = new TLinux_MultiTargetPlatformSettings<FLinuxPlatformProperties<false, false, false, true>>();
+		TargetPlatforms.Add(NoEditorTP);
+		PlatformNameToPlatformSettings.Add(NoEditorTPName + TEXT("_Multi"), NoEditorTP);
+		// Server TP
+		ServerTP = new TLinux_MultiTargetPlatformSettings<FLinuxPlatformProperties<false, true, false, true>>();
+		TargetPlatforms.Add(ServerTP);
+		PlatformNameToPlatformSettings.Add(ServerTPName + TEXT("_Multi"), ServerTP);
+		// Client TP
+		ClientTP = new TLinux_MultiTargetPlatformSettings<FLinuxPlatformProperties<false, false, true, true>>();
+		TargetPlatforms.Add(ClientTP);
+		PlatformNameToPlatformSettings.Add(ClientTPName + TEXT("_Multi"), ClientTP);
 	}
 
 	virtual void GetPlatformSettingsMaps(TMap<FString, ITargetPlatformSettings*>& OutMap) override
@@ -49,4 +104,4 @@ private:
 
 };
 
-IMPLEMENT_MODULE(FLinuxArm64TargetPlatformSettingsModule, LinuxArm64TargetPlatformSettings);
\ No newline at end of file
+IMPLEMENT_MODULE(FLinuxArm64TargetPlatformSettingsModule, LinuxArm64TargetPlatformSettings);
diff --git a/Engine/Source/Developer/Linux/LinuxTargetPlatformControls/Public/LinuxTargetPlatformControls.h b/Engine/Source/Developer/Linux/LinuxTargetPlatformControls/Public/LinuxTargetPlatformControls.h
--- a/Engine/Source/Developer/Linux/LinuxTargetPlatformControls/Public/LinuxTargetPlatformControls.h
+++ b/Engine/Source/Developer/Linux/LinuxTargetPlatformControls/Public/LinuxTargetPlatformControls.h
@@ -7,19 +7,20 @@
 #pragma once
 
 #include "CoreMinimal.h"
+#include "Misc/ConfigCacheIni.h"
 #include "Interfaces/TargetDeviceId.h"
 #include "Common/TargetPlatformControlsBase.h"
 #include "SteamDeck/SteamDeckDevice.h"
 #include "AnalyticsEventAttribute.h"
 #include "InstalledPlatformInfo.h"
 #include "LinuxTargetDevice.h"
-#include "Linux/LinuxPlatformProperties.h" // IWYU pragma: export
+#include "LinuxTargetPlatformSettings.h"
+#include "Linux/LinuxPlatformProperties.h"
 #if WITH_ENGINE
 #include "TextureResource.h"
 #endif // WITH_ENGINE
 #define LOCTEXT_NAMESPACE "TLinuxTargetPlatformControls"
 
-
 /**
  * Template for Linux target platforms controls
  */
@@ -34,8 +35,14 @@ public:
 	/**
 	 * Default constructor.
 	 */
-	TLinuxTargetPlatformControls(ITargetPlatformSettings* TargetPlatformSettings)
-		: TSuper(TargetPlatformSettings)
+	TLinuxTargetPlatformControls(const PlatformInfo::FTargetPlatformInfo* const InPlatformInfo, ITargetPlatformSettings* TargetPlatformSettings)
+		: TSuper(InPlatformInfo ? InPlatformInfo : new PlatformInfo::FTargetPlatformInfo(
+			TProperties::IniPlatformName(),
+			TProperties::HasEditorOnlyData() ? EBuildTargetType::Editor :
+			TProperties::IsServerOnly() ? EBuildTargetType::Server :
+			TProperties::IsClientOnly() ? EBuildTargetType::Client :
+			EBuildTargetType::Game,
+			TEXT("")), TargetPlatformSettings)
 #if WITH_ENGINE
 		, bChangingDeviceConfig(false)
 #endif // WITH_ENGINE
@@ -60,13 +67,31 @@ public:
 #endif // WITH_ENGINE
 	}
 
+	// Forward this one on to the version that can create a new PlatformInfo for us.  This allows us to specify a PlatformInfo with a CookFlavor in derived
+	// templates, but not have to specify the nullptr manually to keep the code compact
+	TLinuxTargetPlatformControls(ITargetPlatformSettings* TargetPlatformSettings)
+		: TLinuxTargetPlatformControls(nullptr, TargetPlatformSettings)
+	{
+	}
 
 public:
 
 	//~ Begin ITargetPlatform Interface
 
+	virtual FString PlatformName() const override
+	{
+		// instead of TPlatformProperties (which won't have Client for non-desktop platforms), use the Info's name, which is programmatically made
+		return this->PlatformInfo->Name.ToString();
+	}
+
 	virtual void EnableDeviceCheck(bool OnOff) override {}
 
+	virtual FString CookingDeviceProfileName() const override
+	{
+		// when cooking for non-desktop platforms, always use the base platform name as the DP to cook with
+		return TEXT("Linux");
+	}
+
 	virtual bool AddDevice(const FString& DeviceName, bool bDefault) override
 	{
 		return AddDevice(DeviceName, TEXT(""), TEXT(""), TEXT(""), bDefault);
@@ -237,6 +262,23 @@ public:
 			// just use the standard texture format name for this texture
 			GetDefaultTextureFormatNamePerLayer(OutFormats.AddDefaulted_GetRef(), this->GetTargetPlatformSettings(), this, InTexture, true, 4, true);
 		}
+
+		// Apply GLES compatibility remaps (e.g. G16->R16F, A1RGB555->RGB555A1)
+		if (!OutFormats.IsEmpty())
+		{
+			TArray<FName>& LayerFormats = OutFormats.Last();
+			for (FName& TextureFormatName : LayerFormats)
+			{
+				for (int32 RemapIndex = 0; RemapIndex < UE_ARRAY_COUNT(LinuxTexFormat::GenericRemap); ++RemapIndex)
+				{
+					if (TextureFormatName == LinuxTexFormat::GenericRemap[RemapIndex][0])
+					{
+						TextureFormatName = LinuxTexFormat::GenericRemap[RemapIndex][1];
+						break;
+					}
+				}
+			}
+		}
 	}
 
 #endif //WITH_ENGINE
@@ -272,7 +314,7 @@ protected:
 		{
 			FString DeviceName, DeviceUser, DevicePass;
 
-			FString DeviceBaseKey(FString::Printf(TEXT("LinuxTargetPlatfrom_%s_Device_%d"), *TSuper::PlatformName(), NumDevices));
+			FString DeviceBaseKey(FString::Printf(TEXT("LinuxTargetPlatfrom_%s_Device_%d"), *PlatformName(), NumDevices));
 			FString DeviceNameKey = DeviceBaseKey + TEXT("_Name");
 			if (!GConfig->GetString(TEXT("/Script/LinuxTargetPlatform.LinuxTargetSettings"), *DeviceNameKey, DeviceName, GEngineIni))
 			{
@@ -357,7 +399,257 @@ protected:
 	// Holds a map of valid devices.
 	TMap<FString, FLinuxTargetDevicePtr> Devices;
 	TArray<ITargetDevicePtr> SteamDevices;
+};
+
+template<typename TProperties>
+class TLinuxDXTTargetPlatformControls
+	: public TLinuxTargetPlatformControls<TProperties>
+{
+public:
+	TLinuxDXTTargetPlatformControls(ITargetPlatformSettings* TargetPlatformSettings)
+		: TLinuxTargetPlatformControls<TProperties>(new PlatformInfo::FTargetPlatformInfo(
+			TProperties::IniPlatformName(),
+			TProperties::HasEditorOnlyData() ? EBuildTargetType::Editor :
+			TProperties::IsServerOnly() ? EBuildTargetType::Server :
+			TProperties::IsClientOnly() ? EBuildTargetType::Client :
+			EBuildTargetType::Game,
+			TEXT("DXT")),
+			TargetPlatformSettings)
+	{
+	}
+
+	virtual float GetVariantPriority() const override
+	{
+		float Priority;
+		return (FConfigCacheIni::ForPlatform(TEXT("Linux"))->GetFloat(LINUX_SECTION_TEXT, TEXT("TextureFormatPriority_DXT"), Priority, GEngineIni) ?
+			Priority : 0.6f) * 10.0f + (TProperties::IsClientOnly() ? 0.25f : 0.5f);
+	}
+};
+
+template<typename TProperties>
+class TLinuxASTCTargetPlatformControls
+	: public TLinuxTargetPlatformControls<TProperties>
+{
+public:
+	TLinuxASTCTargetPlatformControls(ITargetPlatformSettings* TargetPlatformSettings)
+		: TLinuxTargetPlatformControls<TProperties>(new PlatformInfo::FTargetPlatformInfo(
+			TProperties::IniPlatformName(),
+			TProperties::HasEditorOnlyData() ? EBuildTargetType::Editor :
+			TProperties::IsServerOnly() ? EBuildTargetType::Server :
+			TProperties::IsClientOnly() ? EBuildTargetType::Client :
+			EBuildTargetType::Game,
+			TEXT("ASTC")),
+			TargetPlatformSettings)
+	{
+	}
+
+	virtual float GetVariantPriority() const override
+	{
+		float Priority;
+		return (FConfigCacheIni::ForPlatform(TEXT("Linux"))->GetFloat(LINUX_SECTION_TEXT, TEXT("TextureFormatPriority_ASTC"), Priority, GEngineIni) ?
+			Priority : 0.9f) * 10.0f + (TProperties::IsClientOnly() ? 0.25f : 0.5f);
+	}
+
+#if WITH_ENGINE
+	virtual void GetTextureFormats(const UTexture* InTexture, TArray< TArray<FName> >& OutFormats) const override
+	{
+		// call base to populate formats and apply GenericRemap
+		TLinuxTargetPlatformControls<TProperties>::GetTextureFormats(InTexture, OutFormats);
+
+		// perform any remapping away from defaults
+		if (!OutFormats.IsEmpty())
+		{
+			TArray<FName>& LayerFormats = OutFormats.Last();
+			for (FName& TextureFormatName : LayerFormats)
+			{
+				for (int32 RemapIndex = 0; RemapIndex < UE_ARRAY_COUNT(LinuxTexFormat::ASTCRemap); ++RemapIndex)
+				{
+					if (TextureFormatName == LinuxTexFormat::ASTCRemap[RemapIndex][0])
+					{
+						TextureFormatName = LinuxTexFormat::ASTCRemap[RemapIndex][1];
+						break;
+					}
+				}
+			}
+
+			bool bSupportASTCHDR = this->GetTargetPlatformSettings()->UsesASTCHDR();
+
+			if (!bSupportASTCHDR)
+			{
+				for (FName& TextureFormatName : LayerFormats)
+				{
+					if (TextureFormatName == LinuxTexFormat::NameASTC_RGB_HDR)
+					{
+						TextureFormatName = LinuxTexFormat::NameRGBA16F;
+					}
+				}
+			}
+		}
+	}
+
+#endif //WITH_ENGINE
+};
+
+template<typename TProperties>
+class TLinuxETC2TargetPlatformControls
+	: public TLinuxTargetPlatformControls<TProperties>
+{
+public:
+	TLinuxETC2TargetPlatformControls(ITargetPlatformSettings* TargetPlatformSettings)
+		: TLinuxTargetPlatformControls<TProperties>(new PlatformInfo::FTargetPlatformInfo(
+			TProperties::IniPlatformName(),
+			TProperties::HasEditorOnlyData() ? EBuildTargetType::Editor :
+			TProperties::IsServerOnly() ? EBuildTargetType::Server :
+			TProperties::IsClientOnly() ? EBuildTargetType::Client :
+			EBuildTargetType::Game,
+			TEXT("ETC2")),
+			TargetPlatformSettings)
+	{
+	}
+
+	virtual float GetVariantPriority() const override
+	{
+		float Priority;
+		return (FConfigCacheIni::ForPlatform(TEXT("Linux"))->GetFloat(LINUX_SECTION_TEXT, TEXT("TextureFormatPriority_ETC2"), Priority, GEngineIni) ?
+			Priority : 0.2f) * 10.0f + (TProperties::IsClientOnly() ? 0.25f : 0.5f);
+	}
+
+#if WITH_ENGINE
+	virtual void GetTextureFormats(const UTexture* InTexture, TArray< TArray<FName> >& OutFormats) const override
+	{
+		// call base to populate formats and apply GenericRemap
+		TLinuxTargetPlatformControls<TProperties>::GetTextureFormats(InTexture, OutFormats);
+
+		// perform any remapping away from defaults
+		if (!OutFormats.IsEmpty())
+		{
+			TArray<FName>& LayerFormats = OutFormats.Last();
+			for (FName& TextureFormatName : LayerFormats)
+			{
+				for (int32 RemapIndex = 0; RemapIndex < UE_ARRAY_COUNT(LinuxTexFormat::ETCRemap); ++RemapIndex)
+				{
+					if (TextureFormatName == LinuxTexFormat::ETCRemap[RemapIndex][0])
+					{
+						TextureFormatName = LinuxTexFormat::ETCRemap[RemapIndex][1];
+						break;
+					}
+				}
+			}
+		}
+	}
+
+#endif // WITH_ENGINE
+};
+
+template<typename TProperties>
+class TLinuxMultiTargetPlatformControls
+	: public TLinuxTargetPlatformControls<TProperties>
+{
+	TArray<ITargetPlatformControls*> FormatTargetPlatforms;
+	FString FormatTargetString;
+
+public:
+	TLinuxMultiTargetPlatformControls(ITargetPlatformSettings* InTargetPlatformSettings)
+		: TLinuxTargetPlatformControls<TProperties>(new PlatformInfo::FTargetPlatformInfo(
+			TProperties::IniPlatformName(),
+			TProperties::HasEditorOnlyData() ? EBuildTargetType::Editor :
+			TProperties::IsServerOnly() ? EBuildTargetType::Server :
+			TProperties::IsClientOnly() ? EBuildTargetType::Client :
+			EBuildTargetType::Game,
+			TEXT("Multi")),
+			InTargetPlatformSettings)
+	{
+	}
 
+	// Set up which formats to include based on bMultiTargetFormat_* config keys.
+	// Pass all single-format target platform controls; this will filter and deduplicate them.
+	void LoadFormats(TArray<ITargetPlatformControls*> SingleFormatTPs)
+	{
+		// sort formats by priority so higher priority formats are packaged (and thus used by the device) first
+		// use GetVariantPriority() which has hardcoded defaults when the config key is absent
+		SingleFormatTPs.Sort([](const ITargetPlatformControls& A, const ITargetPlatformControls& B)
+			{
+				return A.GetVariantPriority() > B.GetVariantPriority();
+			});
+
+		FormatTargetPlatforms.Empty();
+		FormatTargetString = TEXT("");
+
+		// Server multi-targets don't cook visual texture data; nothing to populate.
+		if (TProperties::IsServerOnly())
+		{
+			return;
+		}
+
+		TSet<FString> SeenFormats;
+
+		// Only pick up single-format platforms that match our own type (Game/Client), so that
+		// e.g. MultiClientTP uses Client texture settings rather than the higher-priority Game variants.
+		const EBuildTargetType MyType = TProperties::IsClientOnly() ? EBuildTargetType::Client : EBuildTargetType::Game;
+
+		for (ITargetPlatformControls* SingleFormatTP : SingleFormatTPs)
+		{
+			FString FlavorName = SingleFormatTP->GetTargetPlatformInfo().PlatformFlavor.ToString();
+			if (FlavorName.IsEmpty())
+			{
+				continue;
+			}
+
+			const EBuildTargetType TPType = SingleFormatTP->GetTargetPlatformInfo().PlatformType;
+			// only use each flavor once, match our platform type, and skip server platforms (servers don't cook visual data)
+			if (SeenFormats.Contains(FlavorName) || TPType == EBuildTargetType::Server || TPType != MyType)
+			{
+				continue;
+			}
+			SeenFormats.Add(FlavorName);
+
+			bool bEnabled = false;
+			FString SettingsName = FString(TEXT("bMultiTargetFormat_")) + FlavorName;
+			FConfigCacheIni::ForPlatform(TEXT("Linux"))->GetBool(LINUX_SECTION_TEXT, *SettingsName, bEnabled, GEngineIni);
+			if (bEnabled)
+			{
+				if (FormatTargetPlatforms.Num())
+				{
+					FormatTargetString += TEXT(",");
+				}
+				FormatTargetString += FlavorName;
+				FormatTargetPlatforms.Add(SingleFormatTP);
+			}
+		}
+
+		PlatformInfo::UpdatePlatformDisplayName(
+			FString(TProperties::IniPlatformName()) + TEXT("_Multi"), DisplayName());
+	}
+
+	virtual float GetVariantPriority() const override
+	{
+		// lowest priority so specific variants are chosen first
+		return 0.5f;
+	}
+
+#if WITH_ENGINE
+	virtual void GetTextureFormats(const UTexture* InTexture, TArray< TArray<FName> >& OutFormats) const override
+	{
+		// Ask each enabled single-format variant to choose texture formats
+		for (ITargetPlatformControls* Platform : FormatTargetPlatforms)
+		{
+			TArray< TArray<FName> > PlatformFormats;
+			Platform->GetTextureFormats(InTexture, PlatformFormats);
+			for (TArray<FName>& FormatPerLayer : PlatformFormats)
+			{
+				OutFormats.AddUnique(FormatPerLayer);
+			}
+		}
+	}
+#endif // WITH_ENGINE
+
+	virtual FText DisplayName() const override
+	{
+		return FText::Format(
+			LOCTEXT("Linux_Multi", "{0} (Multi:{1})"),
+			FText::FromString(TProperties::IniPlatformName()),
+			FText::FromString(FormatTargetString));
+	}
 };
 
 #undef LOCTEXT_NAMESPACE
diff --git a/Engine/Source/Developer/Linux/LinuxTargetPlatformSettings/Classes/LinuxTargetSettings.h b/Engine/Source/Developer/Linux/LinuxTargetPlatformSettings/Classes/LinuxTargetSettings.h
--- a/Engine/Source/Developer/Linux/LinuxTargetPlatformSettings/Classes/LinuxTargetSettings.h
+++ b/Engine/Source/Developer/Linux/LinuxTargetPlatformSettings/Classes/LinuxTargetSettings.h
@@ -65,6 +65,39 @@ public:
 	UPROPERTY(config, EditAnywhere, Category = "Renderer", meta = (DisplayName = "Hardware Ray Tracing Mode"))
 	ETargetSettingsRayTracingRuntimeMode RayTracingMode;
 
+	// --- Linux ARM64 Multi texture format cook targets ---
+	// These ARM64-specific settings are intentionally stored in ULinuxTargetSettings because both the Linux (x86-64)
+	// and LinuxArm64 modules share the same ini section (/Script/LinuxTargetPlatform.LinuxTargetSettings).
+	// The LinuxArm64 controls module reads these values from GEngineIni using LINUX_SECTION_TEXT.
+
+	/** Include DXT textures when packaging with the Linux ARM64 (Multi) cook target. */
+	UPROPERTY(GlobalConfig, EditAnywhere, Category = "LinuxArm64 Multi Texture Formats", meta = (DisplayName = "Include DXT textures"))
+	bool bMultiTargetFormat_DXT;
+
+	/** Include ASTC textures when packaging with the Linux ARM64 (Multi) cook target. */
+	UPROPERTY(GlobalConfig, EditAnywhere, Category = "LinuxArm64 Multi Texture Formats", meta = (DisplayName = "Include ASTC textures"))
+	bool bMultiTargetFormat_ASTC;
+
+	/** Include ETC2 textures when packaging with the Linux ARM64 (Multi) cook target. */
+	UPROPERTY(GlobalConfig, EditAnywhere, Category = "LinuxArm64 Multi Texture Formats", meta = (DisplayName = "Include ETC2 textures"))
+	bool bMultiTargetFormat_ETC2;
+
+	/** Priority for the ETC2 texture format when packaging using Linux ARM64 (Multi). The highest priority format supported by the device will be used. Default value is 0.2. */
+	UPROPERTY(GlobalConfig, EditAnywhere, Category = "LinuxArm64 Multi Texture Formats", meta = (DisplayName = "ETC2 texture format priority"))
+	float TextureFormatPriority_ETC2;
+
+	/** Priority for the DXT texture format when packaging using Linux ARM64 (Multi). The highest priority format supported by the device will be used. Default value is 0.6. */
+	UPROPERTY(GlobalConfig, EditAnywhere, Category = "LinuxArm64 Multi Texture Formats", meta = (DisplayName = "DXT texture format priority"))
+	float TextureFormatPriority_DXT;
+
+	/** Priority for the ASTC texture format when packaging using Linux ARM64 (Multi). The highest priority format supported by the device will be used. Default value is 0.9. */
+	UPROPERTY(GlobalConfig, EditAnywhere, Category = "LinuxArm64 Multi Texture Formats", meta = (DisplayName = "ASTC texture format priority"))
+	float TextureFormatPriority_ASTC;
+
+#if WITH_EDITOR
+	LINUXTARGETPLATFORMSETTINGS_API virtual void PostEditChangeProperty(struct FPropertyChangedEvent& PropertyChangedEvent) override;
+#endif
+
 	void virtual OverrideConfigSection(FString& InOutSectionName) override
 	{
 		InOutSectionName = TEXT("/Script/LinuxTargetPlatform.LinuxTargetSettings");
diff --git a/Engine/Source/Developer/Linux/LinuxTargetPlatformSettings/LinuxTargetPlatformSettings.Build.cs b/Engine/Source/Developer/Linux/LinuxTargetPlatformSettings/LinuxTargetPlatformSettings.Build.cs
--- a/Engine/Source/Developer/Linux/LinuxTargetPlatformSettings/LinuxTargetPlatformSettings.Build.cs
+++ b/Engine/Source/Developer/Linux/LinuxTargetPlatformSettings/LinuxTargetPlatformSettings.Build.cs
@@ -33,5 +33,10 @@ public class LinuxTargetPlatformSettings : ModuleRules
 			PrivateDependencyModuleNames.Add("Engine");
 			PrivateIncludePathModuleNames.Add("TextureCompressor");
 		}
+
+		if (Target.bBuildTargetDeveloperTools)
+		{
+			PrivateIncludePathModuleNames.Add("LinuxArm64TargetPlatformControls");
+		}
 	}
 }
diff --git a/Engine/Source/Developer/Linux/LinuxTargetPlatformSettings/Private/LinuxTargetPlatformClasses.cpp b/Engine/Source/Developer/Linux/LinuxTargetPlatformSettings/Private/LinuxTargetPlatformClasses.cpp
--- a/Engine/Source/Developer/Linux/LinuxTargetPlatformSettings/Private/LinuxTargetPlatformClasses.cpp
+++ b/Engine/Source/Developer/Linux/LinuxTargetPlatformSettings/Private/LinuxTargetPlatformClasses.cpp
@@ -8,6 +8,11 @@
 #include "LinuxTargetSettings.h"
 #include "LinuxTargetPlatformSettings.h"
 #include "Misc/ConfigCacheIni.h"
+#include "Modules/ModuleManager.h"
+
+#if WITH_EDITOR
+#include "ILinuxArm64TargetPlatformControlsModule.h"
+#endif
 
 
 /* ULinuxTargetSettings structors
@@ -15,12 +20,27 @@
 
 ULinuxTargetSettings::ULinuxTargetSettings( const FObjectInitializer& ObjectInitializer )
 	: Super(ObjectInitializer)
+	, bMultiTargetFormat_DXT(true)
+	, bMultiTargetFormat_ASTC(true)
+	, bMultiTargetFormat_ETC2(true)
+	, TextureFormatPriority_ETC2(0.2f)
+	, TextureFormatPriority_DXT(0.6f)
+	, TextureFormatPriority_ASTC(0.9f)
 { }
 
 void ULinuxTargetSettings::PostInitProperties()
 {
 	Super::PostInitProperties();
 
+	// Guard against manually-edited ini files that set all Multi formats to false.
+	// PostEditChangeProperty is not called on startup, so this is the only recovery path.
+	if (!bMultiTargetFormat_DXT && !bMultiTargetFormat_ASTC && !bMultiTargetFormat_ETC2)
+	{
+		UE_LOG(LogTemp, Warning, TEXT("LinuxTargetSettings: All bMultiTargetFormat_* flags were false (possibly from a manual ini edit); defaulting bMultiTargetFormat_ASTC to true."));
+		bMultiTargetFormat_ASTC = true;
+		UpdateSinglePropertyInConfigFile(GetClass()->FindPropertyByName(GET_MEMBER_NAME_CHECKED(ULinuxTargetSettings, bMultiTargetFormat_ASTC)), GetDefaultConfigFilename());
+	}
+
 	FProperty* RayTracingModeProperty = GetClass()->FindPropertyByName(GET_MEMBER_NAME_CHECKED(ULinuxTargetSettings, RayTracingMode));
 	if (RayTracingModeProperty)
 	{
@@ -38,3 +58,39 @@ void ULinuxTargetSettings::PostInitProperties()
 		}
 	}
 }
+
+#if WITH_EDITOR
+void ULinuxTargetSettings::PostEditChangeProperty(struct FPropertyChangedEvent& PropertyChangedEvent)
+{
+	Super::PostEditChangeProperty(PropertyChangedEvent);
+
+	if (PropertyChangedEvent.Property != nullptr &&
+		(PropertyChangedEvent.Property->GetName().StartsWith(TEXT("bMultiTargetFormat")) ||
+		 PropertyChangedEvent.Property->GetName().StartsWith(TEXT("TextureFormatPriority"))))
+	{
+		// Correct in-memory state BEFORE any config writes: ensure at least one format remains selected.
+		// Writing the correction first means any crash between the two config writes leaves the config valid.
+		FProperty* ASTCProperty = GetClass()->FindPropertyByName(GET_MEMBER_NAME_CHECKED(ULinuxTargetSettings, bMultiTargetFormat_ASTC));
+		bool bWroteCorrectedASTC = false;
+		if (!bMultiTargetFormat_DXT && !bMultiTargetFormat_ASTC && !bMultiTargetFormat_ETC2)
+		{
+			bMultiTargetFormat_ASTC = true;
+			UpdateSinglePropertyInConfigFile(ASTCProperty, GetDefaultConfigFilename());
+			bWroteCorrectedASTC = true;
+		}
+
+		// Skip the write if we already wrote this property as part of the correction above
+		if (!bWroteCorrectedASTC || PropertyChangedEvent.Property != ASTCProperty)
+		{
+			UpdateSinglePropertyInConfigFile(PropertyChangedEvent.Property, GetDefaultConfigFilename());
+		}
+
+		// Notify the LinuxArm64 target platform controls module if it's loaded
+		ILinuxArm64TargetPlatformControlsModule* Module = FModuleManager::GetModulePtr<ILinuxArm64TargetPlatformControlsModule>("LinuxArm64TargetPlatformControls");
+		if (Module)
+		{
+			Module->NotifyMultiSelectedFormatsChanged();
+		}
+	}
+}
+#endif
diff --git a/Engine/Source/Developer/Linux/LinuxTargetPlatformSettings/Public/LinuxTargetPlatformSettings.h b/Engine/Source/Developer/Linux/LinuxTargetPlatformSettings/Public/LinuxTargetPlatformSettings.h
--- a/Engine/Source/Developer/Linux/LinuxTargetPlatformSettings/Public/LinuxTargetPlatformSettings.h
+++ b/Engine/Source/Developer/Linux/LinuxTargetPlatformSettings/Public/LinuxTargetPlatformSettings.h
@@ -18,6 +18,79 @@
 
 class UTextureLODSettings;
 
+namespace LinuxTexFormat
+{
+	// Compressed Texture Formats
+ 	const static FName NameDXT1(TEXT("DXT1"));
+ 	const static FName NameDXT5(TEXT("DXT5"));
+ 	const static FName NameDXT5n(TEXT("DXT5n"));
+ 	const static FName NameAutoDXT(TEXT("AutoDXT"));
+ 	const static FName NameBC4(TEXT("BC4"));
+ 	const static FName NameBC5(TEXT("BC5"));
+ 	const static FName NameBC6H(TEXT("BC6H"));
+ 	const static FName NameBC7(TEXT("BC7"));
+
+	const static FName NameETC2_RGB(TEXT("ETC2_RGB"));
+	const static FName NameETC2_RGBA(TEXT("ETC2_RGBA"));
+	const static FName NameETC2_R11(TEXT("ETC2_R11"));
+	const static FName NameETC2_RG11(TEXT("ETC2_RG11"));
+	const static FName NameAutoETC2(TEXT("AutoETC2"));
+
+	const static FName NameAutoASTC(TEXT("ASTC_RGBAuto"));
+	const static FName NameASTC_NormalRG(TEXT("ASTC_NormalRG"));
+	// L+A mode supported by ARM ASTC encoder. Not included in ASTCRemap as it is already an ASTC format and requires no remapping.
+	const static FName NameASTC_NormalLA(TEXT("ASTC_NormalLA"));
+
+	// Uncompressed Texture Formats
+	const static FName NameBGRA8(TEXT("BGRA8"));
+	const static FName NameG8(TEXT("G8"));
+	const static FName NameRGBA16F(TEXT("RGBA16F"));
+	const static FName NameR16F(TEXT("R16F"));
+	const static FName NameG16(TEXT("G16"));
+
+	//A1RGB555 is mapped to RGB555A1, because OpenGL GL_RGB5_A1 only supports alpha on the lowest bit
+	const static FName NameA1RGB555(TEXT("A1RGB555"));
+	const static FName NameRGB555A1(TEXT("RGB555A1"));
+
+	const static FName GenericRemap[][2] =
+	{
+		{ NameA1RGB555,		NameRGB555A1		},
+		{ NameG16,		NameR16F		}, // GLES does not support R16Unorm, fallback all to R16F
+	};
+
+	const static FName NameASTC_RGB_HDR(TEXT("ASTC_RGB_HDR"));
+	const static FName NameASTC_RGB(TEXT("ASTC_RGB"));
+	const static FName NameASTC_RGBA(TEXT("ASTC_RGBA"));
+	const static FName NameASTC_RGBA_HQ(TEXT("ASTC_RGBA_HQ"));
+	const static FName NameASTC_NormalAG(TEXT("ASTC_NormalAG"));
+
+	const static FName ASTCRemap[][2] =
+	{
+		// Default format:		Preferred format:
+		{ NameDXT1,			NameASTC_RGB			},
+		{ NameDXT5,			NameASTC_RGBA			},
+		{ NameDXT5n,			NameASTC_NormalAG		},
+		{ NameBC5,			NameASTC_NormalRG		},
+		{ NameBC4,			NameETC2_R11			},	// No ASTC single-channel format exists; ETC2_R11 is the closest equivalent
+		{ NameBC6H,			NameASTC_RGB_HDR		},
+		{ NameBC7,			NameASTC_RGBA_HQ		},
+		{ NameAutoDXT,			NameAutoASTC			},
+	};
+
+	const static FName ETCRemap[][2] =
+	{
+		// Default format:	ETC2 format:
+		{ NameDXT1,			NameETC2_RGB	},
+		{ NameDXT5,			NameETC2_RGBA	},
+		{ NameDXT5n,			NameETC2_RGB	},	// this is a quality/compatibility tradeoff that matches what Android does
+		{ NameBC5,			NameETC2_RG11	},
+		{ NameBC4,			NameETC2_R11	},
+		{ NameBC6H,			NameRGBA16F	},
+		{ NameBC7,			NameETC2_RGBA	},
+		{ NameAutoDXT,			NameAutoETC2	},
+	};
+}
+
 /**
  * Template for Linux target platforms settings
  */
@@ -32,8 +105,8 @@ public:
 	/**
 	 * Default constructor.
 	 */
-	TLinuxTargetPlatformSettings( )
-
+	TLinuxTargetPlatformSettings(const TCHAR* CookFlavor = nullptr, const TCHAR* OverrideIniPlatformName = nullptr)
+		: TSuper(CookFlavor, OverrideIniPlatformName)
 	{
 #if WITH_ENGINE
 		TextureLODSettings = nullptr;
@@ -184,4 +257,42 @@ protected:
 #endif // WITH_ENGINE
 };
 
+template<typename TProperties>
+class TLinux_DXTTargetPlatformSettings : public TLinuxTargetPlatformSettings<TProperties>
+{
+public:
+	TLinux_DXTTargetPlatformSettings() : TLinuxTargetPlatformSettings<TProperties>(TEXT("DXT"))
+	{
+	}
+};
+
+template<typename TProperties>
+class TLinux_ASTCTargetPlatformSettings : public TLinuxTargetPlatformSettings<TProperties>
+{
+public:
+	TLinux_ASTCTargetPlatformSettings() : TLinuxTargetPlatformSettings<TProperties>(TEXT("ASTC"))
+	{
+	}
+};
+
+template<typename TProperties>
+class TLinux_ETC2TargetPlatformSettings : public TLinuxTargetPlatformSettings<TProperties>
+{
+public:
+
+	TLinux_ETC2TargetPlatformSettings() : TLinuxTargetPlatformSettings<TProperties>(TEXT("ETC2"))
+	{
+	}
+};
+
+template<typename TProperties>
+class TLinux_MultiTargetPlatformSettings : public TLinuxTargetPlatformSettings<TProperties>
+{
+public:
+
+	TLinux_MultiTargetPlatformSettings() : TLinuxTargetPlatformSettings<TProperties>(TEXT("Multi"))
+	{
+	}
+};
+
 #undef LOCTEXT_NAMESPACE
diff --git a/Engine/Source/Programs/AutomationTool/Linux/LinuxPlatform.Automation.cs b/Engine/Source/Programs/AutomationTool/Linux/LinuxPlatform.Automation.cs
--- a/Engine/Source/Programs/AutomationTool/Linux/LinuxPlatform.Automation.cs
+++ b/Engine/Source/Programs/AutomationTool/Linux/LinuxPlatform.Automation.cs
@@ -467,3 +467,76 @@ public class GenericLinuxPlatformArm64 : BaseLinuxPlatform
 	{
 	}
 }
+
+public class LinuxPlatformArm64DXT : GenericLinuxPlatformArm64
+{
+	public LinuxPlatformArm64DXT()
+		: base()
+	{
+	}
+
+	public override TargetPlatformDescriptor GetTargetPlatformDescriptor()
+	{
+		return new TargetPlatformDescriptor(TargetPlatformType, "DXT");
+	}
+
+	public override string GetCookPlatform(bool bDedicatedServer, bool bIsClientOnly)
+	{
+		return base.GetCookPlatform(bDedicatedServer, bIsClientOnly) + "_DXT";
+	}
+}
+
+public class LinuxPlatformArm64ASTC : GenericLinuxPlatformArm64
+{
+	public LinuxPlatformArm64ASTC()
+		: base()
+	{
+	}
+	
+	public override TargetPlatformDescriptor GetTargetPlatformDescriptor()
+	{
+		return new TargetPlatformDescriptor(TargetPlatformType, "ASTC");
+	}
+
+	public override string GetCookPlatform(bool bDedicatedServer, bool bIsClientOnly)
+	{
+		return base.GetCookPlatform(bDedicatedServer, bIsClientOnly) + "_ASTC";
+	}
+}
+
+public class LinuxPlatformArm64ETC2 : GenericLinuxPlatformArm64
+{
+	public LinuxPlatformArm64ETC2()
+		: base()
+	{
+	}
+
+	public override TargetPlatformDescriptor GetTargetPlatformDescriptor()
+	{
+		return new TargetPlatformDescriptor(TargetPlatformType, "ETC2");
+	}
+
+	public override string GetCookPlatform(bool bDedicatedServer, bool bIsClientOnly)
+	{
+		return base.GetCookPlatform(bDedicatedServer, bIsClientOnly) + "_ETC2";
+	}
+}
+
+public class LinuxPlatformArm64Multi : GenericLinuxPlatformArm64
+{
+	public LinuxPlatformArm64Multi()
+		: base()
+	{
+	}
+
+	public override TargetPlatformDescriptor GetTargetPlatformDescriptor()
+	{
+		return new TargetPlatformDescriptor(TargetPlatformType, "Multi");
+	}
+
+	public override string GetCookPlatform(bool bDedicatedServer, bool bIsClientOnly)
+	{
+		return base.GetCookPlatform(bDedicatedServer, bIsClientOnly) + "_Multi";
+	}
+}
+
diff --git a/Engine/Source/Runtime/VulkanRHI/Private/Linux/VulkanLinuxPlatform.h b/Engine/Source/Runtime/VulkanRHI/Private/Linux/VulkanLinuxPlatform.h
--- a/Engine/Source/Runtime/VulkanRHI/Private/Linux/VulkanLinuxPlatform.h
+++ b/Engine/Source/Runtime/VulkanRHI/Private/Linux/VulkanLinuxPlatform.h
@@ -50,6 +50,9 @@ public:
 
 	static VkTimeDomainKHR GetTimeDomain() { return VK_TIME_DOMAIN_CLOCK_MONOTONIC_KHR; }
 
+	static bool SupportsASTCTextureFormats() { return PLATFORM_LINUXARM64; }
+	static bool SupportsETC2TextureFormats() { return PLATFORM_LINUXARM64; }
+
 protected:
 	static void* VulkanLib;
 	static bool bLoadedSDLVulkanLibrary;
